Portal: Removed placeholder messaging code in preparation for message type trait refactor

diff --git a/Engine/Source/Runtime/Portal/Proxies/PortalPackageProxies/Private/PackageInstallerProxy.cpp b/Engine/Source/Runtime/Portal/Proxies/PortalPackageProxies/Private/PackageInstallerProxy.cpp
--- a/Engine/Source/Runtime/Portal/Proxies/PortalPackageProxies/Private/PackageInstallerProxy.cpp
+++ b/Engine/Source/Runtime/Portal/Proxies/PortalPackageProxies/Private/PackageInstallerProxy.cpp
@@ -10,9 +10,11 @@
 
 FPackageInstallerProxy::FPackageInstallerProxy()
 {
+/*
 	MessageEndpoint = FMessageEndpoint::Builder("FPackageInstallerProxy")
 		.Handling<FPortalPackageInstallResponse>(this, &FPackageInstallerProxy::HandlePackageInstallResponse)
 		.Handling<FPortalPackageUninstallResponse>(this, &FPackageInstallerProxy::HandlePackageUninstallResponse);
+*/
 }
 
 
diff --git a/Engine/Source/Runtime/Portal/Proxies/PortalPackageProxies/Public/PortalPackageInstallerMessages.h b/Engine/Source/Runtime/Portal/Proxies/PortalPackageProxies/Public/PortalPackageInstallerMessages.h
--- a/Engine/Source/Runtime/Portal/Proxies/PortalPackageProxies/Public/PortalPackageInstallerMessages.h
+++ b/Engine/Source/Runtime/Portal/Proxies/PortalPackageProxies/Public/PortalPackageInstallerMessages.h
@@ -33,8 +33,6 @@ struct FPortalPackageInstallRequest
 	FGuid RequestId;
 };
 
-template<> struct TStructOpsTypeTraits<FPortalPackageInstallRequest> : public TStructOpsTypeTraitsBase { enum { WithMessageHandling = true }; };
-
 
 /**
  * Implements a reply message to Portal package installation requests.
@@ -52,8 +50,6 @@ struct FPortalPackageInstallResponse
 	bool Success;
 };
 
-template<> struct TStructOpsTypeTraits<FPortalPackageInstallResponse> : public TStructOpsTypeTraitsBase{ enum { WithMessageHandling = true }; };
-
 
 /**
  * Implements a message for Portal package uninstallation.
@@ -84,8 +80,6 @@ struct FPortalPackageUninstallRequest
 	FGuid RequestId;
 };
 
-template<> struct TStructOpsTypeTraits<FPortalPackageUninstallRequest> : public TStructOpsTypeTraitsBase{ enum { WithMessageHandling = true }; };
-
 
 /**
  * Implements a reply message to Portal package uninstallation requests.
@@ -102,5 +96,3 @@ struct FPortalPackageUninstallResponse
 	UPROPERTY()
 	bool Success;
 };
-
-template<> struct TStructOpsTypeTraits<FPortalPackageUninstallResponse> : public TStructOpsTypeTraitsBase{ enum { WithMessageHandling = true }; };
