Portal: Removed obsolete files; renamed Module.cpp files

diff --git a/Engine/Source/Runtime/Portal/Proxies/PortalAccountProxies/PortalAccountProxies.Build.cs b/Engine/Source/Runtime/Portal/Proxies/PortalAccountProxies/PortalAccountProxies.Build.cs
deleted file mode 100644
diff --git a/Engine/Source/Runtime/Portal/Proxies/PortalAccountProxies/Private/Module.cpp b/Engine/Source/Runtime/Portal/Proxies/PortalAccountProxies/Private/Module.cpp
deleted file mode 100644
diff --git a/Engine/Source/Runtime/Portal/Proxies/PortalAccountProxies/Private/PrivatePCH.h b/Engine/Source/Runtime/Portal/Proxies/PortalAccountProxies/Private/PrivatePCH.h
deleted file mode 100644
diff --git a/Engine/Source/Runtime/Portal/Proxies/PortalChatProxies/PortalChatProxies.Build.cs b/Engine/Source/Runtime/Portal/Proxies/PortalChatProxies/PortalChatProxies.Build.cs
deleted file mode 100644
diff --git a/Engine/Source/Runtime/Portal/Proxies/PortalChatProxies/Private/Module.cpp b/Engine/Source/Runtime/Portal/Proxies/PortalChatProxies/Private/Module.cpp
deleted file mode 100644
diff --git a/Engine/Source/Runtime/Portal/Proxies/PortalChatProxies/Private/PrivatePCH.h b/Engine/Source/Runtime/Portal/Proxies/PortalChatProxies/Private/PrivatePCH.h
deleted file mode 100644
diff --git a/Engine/Source/Runtime/Portal/Proxies/PortalPackageProxies/PortalPackageProxies.Build.cs b/Engine/Source/Runtime/Portal/Proxies/PortalPackageProxies/PortalPackageProxies.Build.cs
deleted file mode 100644
diff --git a/Engine/Source/Runtime/Portal/Proxies/PortalPackageProxies/Private/Module.cpp b/Engine/Source/Runtime/Portal/Proxies/PortalPackageProxies/Private/Module.cpp
deleted file mode 100644
diff --git a/Engine/Source/Runtime/Portal/Proxies/PortalPackageProxies/Private/PackageInstallerProxy.cpp b/Engine/Source/Runtime/Portal/Proxies/PortalPackageProxies/Private/PackageInstallerProxy.cpp
deleted file mode 100644
diff --git a/Engine/Source/Runtime/Portal/Proxies/PortalPackageProxies/Private/PackageInstallerProxy.h b/Engine/Source/Runtime/Portal/Proxies/PortalPackageProxies/Private/PackageInstallerProxy.h
deleted file mode 100644
diff --git a/Engine/Source/Runtime/Portal/Proxies/PortalPackageProxies/Private/PrivatePCH.h b/Engine/Source/Runtime/Portal/Proxies/PortalPackageProxies/Private/PrivatePCH.h
deleted file mode 100644
diff --git a/Engine/Source/Runtime/Portal/Proxies/PortalPackageProxies/Public/PortalPackageInstallerMessages.h b/Engine/Source/Runtime/Portal/Proxies/PortalPackageProxies/Public/PortalPackageInstallerMessages.h
deleted file mode 100644
diff --git a/Engine/Source/Runtime/Portal/Services/Private/Module.cpp b/Engine/Source/Runtime/Portal/Services/Private/Module.cpp
deleted file mode 100644
diff --git a/Engine/Source/Runtime/Portal/Services/Private/PortalServicesModule.cpp b/Engine/Source/Runtime/Portal/Services/Private/PortalServicesModule.cpp
new file mode 100644
--- /dev/null
+++ b/Engine/Source/Runtime/Portal/Services/Private/PortalServicesModule.cpp
@@ -0,0 +1,52 @@
+// Copyright 1998-2015 Epic Games, Inc. All Rights Reserved.
+
+#include "PortalServicesPrivatePCH.h"
+#include "TypeContainer.h"
+#include "ModuleManager.h"
+
+
+/**
+ * Implements the PortalServices module.
+ */
+class FPortalServicesModule
+	: public IPortalServicesModule
+{
+public:
+
+	/** Virtual destructor. */
+	virtual ~FPortalServicesModule() { }
+
+public:
+
+	// IPortalServicesModule
+
+	virtual FTypeContainer& GetServiceContainer() override
+	{
+		return ServiceContainer;
+	}
+
+public:
+
+	// IModuleInterface interface
+
+	virtual void StartupModule() override
+	{
+	}
+
+	virtual void ShutdownModule() override
+	{
+	}
+
+	virtual bool SupportsDynamicReloading() override
+	{
+		return false;
+	}
+
+private:
+
+	/** Holds registered service instances. */
+	FTypeContainer ServiceContainer;
+};
+
+
+IMPLEMENT_MODULE(FPortalServicesModule, PortalServices);
