Portal: made PortalServices module include-only

diff --git a/Engine/Source/Runtime/Portal/Services/Private/PortalServicesModule.cpp b/Engine/Source/Runtime/Portal/Services/Private/PortalServicesModule.cpp
--- a/Engine/Source/Runtime/Portal/Services/Private/PortalServicesModule.cpp
+++ b/Engine/Source/Runtime/Portal/Services/Private/PortalServicesModule.cpp
@@ -1,7 +1,6 @@
 // Copyright 1998-2015 Epic Games, Inc. All Rights Reserved.
 
 #include "PortalServicesPrivatePCH.h"
-#include "TypeContainer.h"
 #include "ModuleManager.h"
 
 
@@ -9,43 +8,20 @@
  * Implements the PortalServices module.
  */
 class FPortalServicesModule
-	: public IPortalServicesModule
+	: public IModuleInterface
 {
 public:
 
 	/** Virtual destructor. */
 	virtual ~FPortalServicesModule() { }
 
-public:
-
-	// IPortalServicesModule
-
-	virtual FTypeContainer& GetServiceContainer() override
-	{
-		return ServiceContainer;
-	}
-
 public:
 
 	// IModuleInterface interface
 
-	virtual void StartupModule() override
-	{
-	}
-
-	virtual void ShutdownModule() override
-	{
-	}
-
-	virtual bool SupportsDynamicReloading() override
-	{
-		return false;
-	}
-
-private:
-
-	/** Holds registered service instances. */
-	FTypeContainer ServiceContainer;
+	virtual void StartupModule() override { }
+	virtual void ShutdownModule() override { }
+	virtual bool SupportsDynamicReloading() override { return true; }
 };
 
 
diff --git a/Engine/Source/Runtime/Portal/Services/Private/PortalServicesPrivatePCH.h b/Engine/Source/Runtime/Portal/Services/Private/PortalServicesPrivatePCH.h
--- a/Engine/Source/Runtime/Portal/Services/Private/PortalServicesPrivatePCH.h
+++ b/Engine/Source/Runtime/Portal/Services/Private/PortalServicesPrivatePCH.h
@@ -7,4 +7,3 @@
 *****************************************************************************/
 
 #include "Core.h"
-#include "IPortalServicesModule.h"
diff --git a/Engine/Source/Runtime/Portal/Services/Public/IPortalServicesModule.h b/Engine/Source/Runtime/Portal/Services/Public/IPortalServicesModule.h
deleted file mode 100644
