Portal: Implemented RPC messages

diff --git a/Engine/Source/Runtime/Portal/Messages/PortalMessages.Build.cs b/Engine/Source/Runtime/Portal/Messages/PortalMessages.Build.cs
new file mode 100644
--- /dev/null
+++ b/Engine/Source/Runtime/Portal/Messages/PortalMessages.Build.cs
@@ -0,0 +1,25 @@
+// Copyright 1998-2015 Epic Games, Inc. All Rights Reserved.
+
+namespace UnrealBuildTool.Rules
+{
+	public class PortalMessages : ModuleRules
+	{
+		public PortalMessages(TargetInfo Target)
+		{
+			PublicDependencyModuleNames.AddRange(
+				new string[]
+				{
+					"Core",
+					"CoreUObject",
+                    "MessagingRpc",
+				}
+			);
+
+			PrivateIncludePaths.AddRange(
+				new string[] {
+					"Runtime/Portal/Messages/Private",
+				}
+			);
+		}
+	}
+}
diff --git a/Engine/Source/Runtime/Portal/Messages/Private/PortalMessagesModule.cpp b/Engine/Source/Runtime/Portal/Messages/Private/PortalMessagesModule.cpp
new file mode 100644
--- /dev/null
+++ b/Engine/Source/Runtime/Portal/Messages/Private/PortalMessagesModule.cpp
@@ -0,0 +1,27 @@
+// Copyright 1998-2015 Epic Games, Inc. All Rights Reserved.
+
+#include "PortalMessagesPrivatePCH.h"
+#include "ModuleInterface.h"
+
+
+/**
+ * Implements the module.
+ */
+class FModule
+	: public IModuleInterface
+{
+public:
+
+	// IModuleInterface interface
+
+	virtual void StartupModule() override { }
+	virtual void ShutdownModule() override { }
+
+	virtual bool SupportsDynamicReloading() override
+	{
+		return true;
+	}
+};
+
+
+IMPLEMENT_MODULE(FModule, PortalMessages);
diff --git a/Engine/Source/Runtime/Portal/Messages/Private/PortalMessagesPrivatePCH.h b/Engine/Source/Runtime/Portal/Messages/Private/PortalMessagesPrivatePCH.h
new file mode 100644
--- /dev/null
+++ b/Engine/Source/Runtime/Portal/Messages/Private/PortalMessagesPrivatePCH.h
@@ -0,0 +1,9 @@
+// Copyright 1998-2015 Epic Games, Inc. All Rights Reserved.
+
+#pragma once
+
+
+/* Private dependencies
+ *****************************************************************************/
+
+#include "Core.h"
diff --git a/Engine/Source/Runtime/Portal/Messages/Public/PortalAppWindowMessages.h b/Engine/Source/Runtime/Portal/Messages/Public/PortalAppWindowMessages.h
new file mode 100644
--- /dev/null
+++ b/Engine/Source/Runtime/Portal/Messages/Public/PortalAppWindowMessages.h
@@ -0,0 +1,122 @@
+// Copyright 1998-2015 Epic Games, Inc. All Rights Reserved.
+
+#pragma once
+
+#include "RpcMessage.h"
+#include "PortalAppWindowMessages.generated.h"
+
+
+USTRUCT()
+struct FPortalAppWindowOpenFriendsRequest
+	: public FRpcMessage
+{
+	GENERATED_USTRUCT_BODY()
+
+	FPortalAppWindowOpenFriendsRequest() { }
+};
+
+
+USTRUCT()
+struct FPortalAppWindowOpenFriendsResponse
+	: public FRpcMessage
+{
+	GENERATED_USTRUCT_BODY()
+
+	UPROPERTY()
+	bool Result;
+
+	FPortalAppWindowOpenFriendsResponse() { }
+	FPortalAppWindowOpenFriendsResponse(bool InResult)
+		: Result(InResult)
+	{ }
+};
+
+
+USTRUCT()
+struct FPortalAppWindowOpenLibraryRequest
+	: public FRpcMessage
+{
+	GENERATED_USTRUCT_BODY()
+
+	FPortalAppWindowOpenLibraryRequest() { }
+};
+
+
+USTRUCT()
+struct FPortalAppWindowOpenLibraryResponse
+	: public FRpcMessage
+{
+	GENERATED_USTRUCT_BODY()
+
+	UPROPERTY()
+	bool Result;
+
+	FPortalAppWindowOpenLibraryResponse() { }
+	FPortalAppWindowOpenLibraryResponse(bool InResult)
+		: Result(InResult)
+	{ }
+};
+
+
+USTRUCT()
+struct FPortalAppWindowOpenMarketplaceRequest
+	: public FRpcMessage
+{
+	GENERATED_USTRUCT_BODY()
+
+	FPortalAppWindowOpenMarketplaceRequest() { }
+};
+
+
+USTRUCT()
+struct FPortalAppWindowOpenMarketplaceResponse
+	: public FRpcMessage
+{
+	GENERATED_USTRUCT_BODY()
+
+	UPROPERTY()
+	bool Result;
+
+	FPortalAppWindowOpenMarketplaceResponse() { }
+	FPortalAppWindowOpenMarketplaceResponse(bool InResult)
+		: Result(InResult)
+	{ }
+};
+
+
+USTRUCT()
+struct FPortalAppWindowOpenMarketItemRequest
+	: public FRpcMessage
+{
+	GENERATED_USTRUCT_BODY()
+
+	UPROPERTY()
+	FString ItemId;
+
+	FPortalAppWindowOpenMarketItemRequest() { }
+	FPortalAppWindowOpenMarketItemRequest(const FString& InItemId)
+		: ItemId(InItemId)
+	{ }
+};
+
+
+USTRUCT()
+struct FPortalAppWindowOpenMarketItemResponse
+	: public FRpcMessage
+{
+	GENERATED_USTRUCT_BODY()
+
+	UPROPERTY()
+	bool Result;
+
+	FPortalAppWindowOpenMarketItemResponse() { }
+	FPortalAppWindowOpenMarketItemResponse(bool InResult)
+		: Result(InResult)
+	{ }
+};
+
+
+DECLARE_RPC(FPortalAppWindowOpenFriends, bool)
+DECLARE_RPC(FPortalAppWindowOpenLibrary, bool)
+DECLARE_RPC(FPortalAppWindowOpenMarketplace, bool)
+DECLARE_RPC(FPortalAppWindowOpenMarketItem, bool)
diff --git a/Engine/Source/Runtime/Portal/Messages/Public/PortalPackageInstallerMessages.h b/Engine/Source/Runtime/Portal/Messages/Public/PortalPackageInstallerMessages.h
new file mode 100644
--- /dev/null
+++ b/Engine/Source/Runtime/Portal/Messages/Public/PortalPackageInstallerMessages.h
@@ -0,0 +1,97 @@
+// Copyright 1998-2015 Epic Games, Inc. All Rights Reserved.
+
+#pragma once
+
+#include "RpcMessage.h"
+#include "PortalPackageInstallerMessages.generated.h"
+
+
+/* Service discovery messages
+ *****************************************************************************/
+
+USTRUCT()
+struct FPortalPackageInstallerInstallRequest
+	: public FRpcMessage
+{
+	GENERATED_USTRUCT_BODY()
+
+	UPROPERTY()
+	FString AppName;
+
+	UPROPERTY()
+	FString BuildLabel;
+
+	UPROPERTY()
+	FString DestinationPath;
+
+	FPortalPackageInstallerInstallRequest() { }
+	FPortalPackageInstallerInstallRequest(const FString& InAppName, const FString& InBuildLabel, const FString& InDestinationPath)
+		: AppName(InAppName)
+		, BuildLabel(InBuildLabel)
+		, DestinationPath(InDestinationPath)
+	{ }
+};
+
+
+USTRUCT()
+struct FPortalPackageInstallerInstallResponse
+	: public FRpcMessage
+{
+	GENERATED_USTRUCT_BODY()
+
+	UPROPERTY()
+	bool Result;
+
+	FPortalPackageInstallerInstallResponse() { }
+	FPortalPackageInstallerInstallResponse(bool InResult)
+		: Result(InResult)
+	{ }
+};
+
+
+USTRUCT()
+struct FPortalPackageInstallerUninstallRequest
+	: public FRpcMessage
+{
+	GENERATED_USTRUCT_BODY()
+
+	UPROPERTY()
+	FString AppName;
+
+	UPROPERTY()
+	FString BuildLabel;
+
+	UPROPERTY()
+	FString InstallationPath;
+
+	UPROPERTY()
+	bool RemoveUserFiles;
+
+	FPortalPackageInstallerUninstallRequest() { }
+	FPortalPackageInstallerUninstallRequest(const FString& InAppName, const FString& InBuildLabel, const FString& InInstallationPath, bool InRemoveUserFiles)
+		: AppName(InAppName)
+		, BuildLabel(InBuildLabel)
+		, InstallationPath(InInstallationPath)
+		, RemoveUserFiles(InRemoveUserFiles)
+	{ }
+};
+
+
+USTRUCT()
+struct FPortalPackageInstallerUninstallResponse
+	: public FRpcMessage
+{
+	GENERATED_USTRUCT_BODY()
+
+	UPROPERTY()
+	bool Result;
+
+	FPortalPackageInstallerUninstallResponse() { }
+	FPortalPackageInstallerUninstallResponse(bool InResult)
+		: Result(InResult)
+	{ }
+};
+
+
+DECLARE_RPC(FPortalPackageInstallerInstall, bool)
+DECLARE_RPC(FPortalPackageInstallerUninstall, bool)
