Portal: Added missing include

diff --git a/Engine/Source/Runtime/Portal/Messages/Private/PortalMessagesModule.cpp b/Engine/Source/Runtime/Portal/Messages/Private/PortalMessagesModule.cpp
--- a/Engine/Source/Runtime/Portal/Messages/Private/PortalMessagesModule.cpp
+++ b/Engine/Source/Runtime/Portal/Messages/Private/PortalMessagesModule.cpp
@@ -2,12 +2,13 @@
 
 #include "PortalMessagesPrivatePCH.h"
 #include "ModuleInterface.h"
+#include "ModuleManager.h"
 
 
 /**
  * Implements the module.
  */
-class FModule
+class FPortalMessagesModule
 	: public IModuleInterface
 {
 public:
@@ -24,4 +25,4 @@ public:
 };
 
 
-IMPLEMENT_MODULE(FModule, PortalMessages);
+IMPLEMENT_MODULE(FPortalMessagesModule, PortalMessages);
diff --git a/Engine/Source/Runtime/Portal/Messages/Private/PortalMessagesPrivatePCH.h b/Engine/Source/Runtime/Portal/Messages/Private/PortalMessagesPrivatePCH.h
--- a/Engine/Source/Runtime/Portal/Messages/Private/PortalMessagesPrivatePCH.h
+++ b/Engine/Source/Runtime/Portal/Messages/Private/PortalMessagesPrivatePCH.h
@@ -7,3 +7,4 @@
  *****************************************************************************/
 
 #include "Core.h"
+#include "CoreUObject.h"
