Portal: Implemented RPC layer

diff --git a/Engine/Source/Runtime/Portal/Rpc/PortalRpc.Build.cs b/Engine/Source/Runtime/Portal/Rpc/PortalRpc.Build.cs
new file mode 100644
--- /dev/null
+++ b/Engine/Source/Runtime/Portal/Rpc/PortalRpc.Build.cs
@@ -0,0 +1,39 @@
+// Copyright 1998-2015 Epic Games, Inc. All Rights Reserved.
+
+namespace UnrealBuildTool.Rules
+{
+	public class PortalRpc : ModuleRules
+	{
+		public PortalRpc(TargetInfo Target)
+		{
+            DynamicallyLoadedModuleNames.AddRange(
+                new string[] {
+                    "Messaging",
+                    "MessagingRpc",
+                }
+            );
+
+			PublicDependencyModuleNames.AddRange(
+				new string[]
+				{
+					"Core",
+					"CoreUObject",
+				}
+			);
+
+            PrivateIncludePathModuleNames.AddRange(
+                new string[] {
+                    "Messaging",
+                    "MessagingRpc",
+                    "PortalServices",
+                }
+            );
+
+			PrivateIncludePaths.AddRange(
+				new string[] {
+					"Runtime/Portal/Rpc/PortalRpc/Private",
+				}
+			);
+		}
+	}
+}
diff --git a/Engine/Source/Runtime/Portal/Rpc/Private/PortalRpcLocator.cpp b/Engine/Source/Runtime/Portal/Rpc/Private/PortalRpcLocator.cpp
new file mode 100644
--- /dev/null
+++ b/Engine/Source/Runtime/Portal/Rpc/Private/PortalRpcLocator.cpp
@@ -0,0 +1,77 @@
+// Copyright 1998-2015 Epic Games, Inc. All Rights Reserved.
+
+#include "PortalRpcPrivatePCH.h"
+#include "EngineVersion.h"
+#include "PortalRpcLocator.h"
+
+
+/* FPortalRpcLocator structors
+ *****************************************************************************/
+
+FPortalRpcLocator::FPortalRpcLocator()
+	: LastServerResponse(FDateTime::MinValue())
+{
+	MessageEndpoint = FMessageEndpoint::Builder("FPortalRpcLocator")
+		.Handling<FPortalRpcServer>(this, &FPortalRpcLocator::HandleMessage);
+
+	TickerHandle = FTicker::GetCoreTicker().AddTicker(FTickerDelegate::CreateRaw(this, &FPortalRpcLocator::HandleTicker), PORTAL_RPC_LOCATE_INTERVAL);
+}
+
+
+FPortalRpcLocator::~FPortalRpcLocator()
+{
+	FTicker::GetCoreTicker().RemoveTicker(TickerHandle);
+}
+
+
+/* IPortalRpcLocator interface
+ *****************************************************************************/
+
+const FMessageAddress& FPortalRpcLocator::GetServerAddress() const
+{
+	return ServerAddress;
+}
+
+
+FSimpleDelegate& FPortalRpcLocator::OnServerLocated()
+{
+	return ServerLocatedDelegate;
+}
+
+
+FSimpleDelegate& FPortalRpcLocator::OnServerLost()
+{
+	return ServerLostDelegate;
+}
+
+
+/* FPortalRpcLocator callbacks
+ *****************************************************************************/
+
+void FPortalRpcLocator::HandleMessage(const FPortalRpcServer& Message, const TSharedRef<IMessageContext, ESPMode::ThreadSafe>& Context)
+{
+	LastServerResponse = FDateTime::UtcNow();
+
+	FMessageAddress NewServerAddress;
+
+	if (FMessageAddress::Parse(Message.ServerAddress, NewServerAddress) && (NewServerAddress != ServerAddress))
+	{
+		ServerAddress = NewServerAddress;
+		ServerLocatedDelegate.ExecuteIfBound();
+	}
+}
+
+
+bool FPortalRpcLocator::HandleTicker(float DeltaTime)
+{
+	if (ServerAddress.IsValid() && (FDateTime::UtcNow() - LastServerResponse > PORTAL_RPC_LOCATE_TIMEOUT))
+	{
+		ServerAddress.Invalidate();
+		ServerLostDelegate.ExecuteIfBound();
+	}
+
+	// @todo sarge: implement actual product GUID
+	MessageEndpoint->Publish(new FPortalRpcLocateServer(FGuid(), GEngineVersion.ToString()), EMessageScope::Network);
+
+	return true;
+}
diff --git a/Engine/Source/Runtime/Portal/Rpc/Private/PortalRpcLocator.h b/Engine/Source/Runtime/Portal/Rpc/Private/PortalRpcLocator.h
new file mode 100644
--- /dev/null
+++ b/Engine/Source/Runtime/Portal/Rpc/Private/PortalRpcLocator.h
@@ -0,0 +1,59 @@
+// Copyright 1998-2015 Epic Games, Inc. All Rights Reserved.
+
+#pragma once
+
+#include "IMessageContext.h"
+#include "IPortalRpcLocator.h"
+
+
+class FMessageEndpoint;
+struct FPortalRpcServer;
+
+
+class FPortalRpcLocator
+	: public IPortalRpcLocator
+{
+public:
+
+	/** Default constructor. */
+	FPortalRpcLocator();
+
+	/** Virtual destructor. */
+	virtual ~FPortalRpcLocator();
+
+public:
+
+	// IPortalRpcLocator interface
+
+	virtual const FMessageAddress& GetServerAddress() const override;
+	virtual FSimpleDelegate& OnServerLocated() override;
+	virtual FSimpleDelegate& OnServerLost() override;
+
+private:
+
+	/** Handles FPortalRpcServer messages. */
+	void HandleMessage(const FPortalRpcServer& Message, const TSharedRef<IMessageContext, ESPMode::ThreadSafe>& Context);
+
+	/** Handles the ticker. */
+	bool HandleTicker(float DeltaTime);
+
+private:
+
+	/** Time at which the RPC server last responded. */
+	FDateTime LastServerResponse;
+
+	/** Message endpoint. */
+	TSharedPtr<FMessageEndpoint, ESPMode::ThreadSafe> MessageEndpoint;
+
+	/** The message address of the located RPC server, or invalid if no server available. */
+	FMessageAddress ServerAddress;
+
+	/** A delegate that is executed when an RPC server has been located. */
+	FSimpleDelegate ServerLocatedDelegate;
+
+	/** A delegate that is executed when the RPC server has been lost. */
+	FSimpleDelegate ServerLostDelegate;
+
+	/** Handle to the registered ticker. */
+	FDelegateHandle TickerHandle;
+};
diff --git a/Engine/Source/Runtime/Portal/Rpc/Private/PortalRpcMessages.h b/Engine/Source/Runtime/Portal/Rpc/Private/PortalRpcMessages.h
new file mode 100644
--- /dev/null
+++ b/Engine/Source/Runtime/Portal/Rpc/Private/PortalRpcMessages.h
@@ -0,0 +1,48 @@
+// Copyright 1998-2015 Epic Games, Inc. All Rights Reserved.
+
+#pragma once
+
+#include "PortalRpcMessages.generated.h"
+
+
+USTRUCT()
+struct FPortalRpcLocateServer
+{
+	GENERATED_USTRUCT_BODY()
+
+	/** The product's unique identifier. */
+	UPROPERTY()
+	FGuid ProductId;
+
+	/** The product's version string. */
+	UPROPERTY()
+	FString ProductVersion;
+
+	/** Default constructor. */
+	FPortalRpcLocateServer() { }
+
+	/** Create and initialize a new instance. */
+	FPortalRpcLocateServer(const FGuid& InProductId, const FString& InProductVersion)
+		: ProductId(InProductId)
+		, ProductVersion(InProductVersion)
+	{ }
+};
+
+
+USTRUCT()
+struct FPortalRpcServer
+{
+	GENERATED_USTRUCT_BODY()
+
+	/** The RPC server's message address as a string. */
+	UPROPERTY()
+	FString ServerAddress;
+
+	/** Default constructor. */
+	FPortalRpcServer() { }
+
+	/** Create and initialize a new instance. */
+	FPortalRpcServer(const FString& InServerAddress)
+		: ServerAddress(InServerAddress)
+	{ }
+};
diff --git a/Engine/Source/Runtime/Portal/Rpc/Private/PortalRpcModule.cpp b/Engine/Source/Runtime/Portal/Rpc/Private/PortalRpcModule.cpp
new file mode 100644
--- /dev/null
+++ b/Engine/Source/Runtime/Portal/Rpc/Private/PortalRpcModule.cpp
@@ -0,0 +1,39 @@
+// Copyright 1998-2015 Epic Games, Inc. All Rights Reserved.
+
+#include "PortalRpcPrivatePCH.h"
+#include "IPortalRpcModule.h"
+
+
+/**
+ * Implements the PortalRpc module.
+ */
+class FPortalRpcModule
+	: public IPortalRpcModule
+{
+public:
+
+	// IModuleInterface interface
+
+	virtual void StartupModule() override { }
+	virtual void ShutdownModule() override { }
+
+	virtual bool SupportsDynamicReloading() override
+	{
+		return false;
+	}
+
+public:
+
+	virtual TSharedRef<IPortalRpcLocator> CreateLocator() override
+	{
+		return MakeShareable(new FPortalRpcLocator);
+	}
+
+	virtual TSharedRef<IPortalRpcProvider> CreateProvider() override
+	{
+		return MakeShareable(new FPortalRpcProvider);
+	}
+};
+
+
+IMPLEMENT_MODULE(FPortalRpcModule, PortalRpc);
diff --git a/Engine/Source/Runtime/Portal/Rpc/Private/PortalRpcPrivatePCH.h b/Engine/Source/Runtime/Portal/Rpc/Private/PortalRpcPrivatePCH.h
new file mode 100644
--- /dev/null
+++ b/Engine/Source/Runtime/Portal/Rpc/Private/PortalRpcPrivatePCH.h
@@ -0,0 +1,31 @@
+// Copyright 1998-2015 Epic Games, Inc. All Rights Reserved.
+
+#pragma once
+
+
+/* Private dependencies
+ *****************************************************************************/
+
+#include "Core.h"
+#include "CoreUObject.h"
+#include "MessageEndpoint.h"
+#include "MessageEndpointBuilder.h"
+#include "Ticker.h"
+
+
+/* Private constants
+ *****************************************************************************/
+
+/** Defines interval at which the Portal RPC server is being located (in seconds). */
+#define PORTAL_RPC_LOCATE_INTERVAL 5.0
+
+/** Defines the time after which Portal RPC servers time out (in seconds). */
+#define PORTAL_RPC_LOCATE_TIMEOUT 15.0
+
+
+/* Private includes
+ *****************************************************************************/
+
+#include "PortalRpcLocator.h"
+#include "PortalRpcMessages.h"
+#include "PortalRpcProvider.h"
diff --git a/Engine/Source/Runtime/Portal/Rpc/Private/PortalRpcProvider.cpp b/Engine/Source/Runtime/Portal/Rpc/Private/PortalRpcProvider.cpp
new file mode 100644
--- /dev/null
+++ b/Engine/Source/Runtime/Portal/Rpc/Private/PortalRpcProvider.cpp
@@ -0,0 +1,40 @@
+// Copyright 1998-2015 Epic Games, Inc. All Rights Reserved.
+
+#include "PortalRpcPrivatePCH.h"
+#include "IMessageRpcServer.h"
+#include "ModuleManager.h"
+
+
+/* FPortalRpcProvider structors
+ *****************************************************************************/
+
+FPortalRpcProvider::FPortalRpcProvider()
+{
+	MessageEndpoint = FMessageEndpoint::Builder("FRpcManager")
+		.Handling<FPortalRpcLocateServer>(this, &FPortalRpcProvider::HandleMessage);
+}
+
+
+/* FPortalRpcProvider callbacks
+ *****************************************************************************/
+
+void FPortalRpcProvider::HandleMessage(const FPortalRpcLocateServer& Message, const TSharedRef<IMessageContext, ESPMode::ThreadSafe>& Context)
+{
+	if (!LookupDelegate.IsBound())
+	{
+		return;
+	}
+
+	const FString ProductKey = Message.ProductId.ToString() + Message.ProductVersion;
+	TSharedPtr<IMessageRpcServer> Server = Servers.FindRef(ProductKey);
+
+	if (!Server.IsValid())
+	{
+		Server = LookupDelegate.Execute(ProductKey);
+	}
+
+	if (Server.IsValid())
+	{
+		MessageEndpoint->Send(new FPortalRpcServer(Server->GetAddress().ToString()), Context->GetSender());
+	}
+}
diff --git a/Engine/Source/Runtime/Portal/Rpc/Private/PortalRpcProvider.h b/Engine/Source/Runtime/Portal/Rpc/Private/PortalRpcProvider.h
new file mode 100644
--- /dev/null
+++ b/Engine/Source/Runtime/Portal/Rpc/Private/PortalRpcProvider.h
@@ -0,0 +1,50 @@
+// Copyright 1998-2015 Epic Games, Inc. All Rights Reserved.
+
+#pragma once
+
+#include "IPortalRpcProvider.h"
+
+
+class FMessageEndpoint;
+struct FPortalRpcLocateServer;
+class IMessageContext;
+
+
+class FPortalRpcProvider
+	: public IPortalRpcProvider
+{
+public:
+
+	/** Default constructor. */
+	FPortalRpcProvider();
+
+	/** Virtual destructor. */
+	virtual ~FPortalRpcProvider() { }
+
+public:
+
+	// IPortalRpcProvider interface
+
+	virtual FOnPortalRpcLookup& OnLookup() override
+	{
+		return LookupDelegate;
+	}
+
+private:
+
+private:
+
+	/** Handles FPortalRpcLocateServer messages. */
+	void HandleMessage(const FPortalRpcLocateServer& Message, const TSharedRef<IMessageContext, ESPMode::ThreadSafe>& Context);
+
+private:
+
+	/** A delegate that is executed when a look-up for an RPC server occurs. */
+	FOnPortalRpcLookup LookupDelegate;
+
+	/** Message endpoint. */
+	TSharedPtr<FMessageEndpoint, ESPMode::ThreadSafe> MessageEndpoint;
+
+	/** Holds the existing RPC servers. */
+	TMap<FString, TSharedPtr<IMessageRpcServer>> Servers;
+};
diff --git a/Engine/Source/Runtime/Portal/Rpc/Public/IPortalRpcLocator.h b/Engine/Source/Runtime/Portal/Rpc/Public/IPortalRpcLocator.h
new file mode 100644
--- /dev/null
+++ b/Engine/Source/Runtime/Portal/Rpc/Public/IPortalRpcLocator.h
@@ -0,0 +1,45 @@
+// Copyright 1998-2015 Epic Games, Inc. All Rights Reserved.
+
+#pragma once
+
+struct FMessageAddress;
+
+
+/**
+ * Interface for Portal RPC server locators.
+ */
+class IPortalRpcLocator
+{
+public:
+
+	/**
+	 * Get the address of the located Portal RPC server.
+	 *
+	 * @return The RPC server's message address, or invalid if no server found.
+	 * @see OnServerLocated, OnServerLost
+	 */
+	virtual const FMessageAddress& GetServerAddress() const = 0;
+
+public:
+
+	/**
+	 * Get a delegate that is executed when an RPC server has been located.
+	 *
+	 * @return The delegate.
+	 * @see GetServerAddress, OnServerLost
+	 */
+	virtual FSimpleDelegate& OnServerLocated() = 0;
+	
+	/**
+	 * Get a delegate that is executed when an RPC server has been located.
+	 *
+	 * @return The delegate.
+	 * @see GetServerAddress, OnServerLocated
+	 */
+	virtual FSimpleDelegate& OnServerLost() = 0;
+
+public:
+
+	/** Virtual destructor. */
+	~IPortalRpcLocator() { }
+};
diff --git a/Engine/Source/Runtime/Portal/Rpc/Public/IPortalRpcModule.h b/Engine/Source/Runtime/Portal/Rpc/Public/IPortalRpcModule.h
new file mode 100644
--- /dev/null
+++ b/Engine/Source/Runtime/Portal/Rpc/Public/IPortalRpcModule.h
@@ -0,0 +1,22 @@
+// Copyright 1998-2015 Epic Games, Inc. All Rights Reserved.
+
+#pragma once
+
+#include "ModuleInterface.h"
+
+
+class IPortalRpcLocator;
+class IPortalRpcProvider;
+
+
+/**
+ * Interface for the PortalRpc module.
+ */
+class IPortalRpcModule
+	: public IModuleInterface
+{
+public:
+
+	virtual TSharedRef<IPortalRpcLocator> CreateLocator() = 0;
+	virtual TSharedRef<IPortalRpcProvider> CreateProvider() = 0;
+};
diff --git a/Engine/Source/Runtime/Portal/Rpc/Public/IPortalRpcProvider.h b/Engine/Source/Runtime/Portal/Rpc/Public/IPortalRpcProvider.h
new file mode 100644
--- /dev/null
+++ b/Engine/Source/Runtime/Portal/Rpc/Public/IPortalRpcProvider.h
@@ -0,0 +1,26 @@
+// Copyright 1998-2015 Epic Games, Inc. All Rights Reserved.
+
+#pragma once
+
+
+class IMessageRpcServer;
+
+
+DECLARE_DELEGATE_RetVal_OneParam(TSharedPtr<IMessageRpcServer>, FOnPortalRpcLookup, const FString& /*ProductKey*/)
+
+
+/**
+ * Interface for Portal RPC endpoint providers.
+ */
+class IPortalRpcProvider
+{
+public:
+
+	/** Get a delegate that is executed when a look-up for an RPC server occurs. */
+	virtual FOnPortalRpcLookup& OnLookup() = 0;
+
+public:
+
+	/** Virtual destructor. */
+	~IPortalRpcProvider() { }
+};
