Portal: Renamed PortalProvider to PortalResponder

diff --git a/Engine/Source/Runtime/Portal/Rpc/Private/PortalRpcModule.cpp b/Engine/Source/Runtime/Portal/Rpc/Private/PortalRpcModule.cpp
--- a/Engine/Source/Runtime/Portal/Rpc/Private/PortalRpcModule.cpp
+++ b/Engine/Source/Runtime/Portal/Rpc/Private/PortalRpcModule.cpp
@@ -29,9 +29,9 @@ public:
 		return MakeShareable(new FPortalRpcLocator);
 	}
 
-	virtual TSharedRef<IPortalRpcProvider> CreateProvider() override
+	virtual TSharedRef<IPortalRpcResponder> CreateResponder() override
 	{
-		return MakeShareable(new FPortalRpcProvider);
+		return MakeShareable(new FPortalRpcResponder);
 	}
 };
 
diff --git a/Engine/Source/Runtime/Portal/Rpc/Private/PortalRpcPrivatePCH.h b/Engine/Source/Runtime/Portal/Rpc/Private/PortalRpcPrivatePCH.h
--- a/Engine/Source/Runtime/Portal/Rpc/Private/PortalRpcPrivatePCH.h
+++ b/Engine/Source/Runtime/Portal/Rpc/Private/PortalRpcPrivatePCH.h
@@ -28,4 +28,4 @@
 
 #include "PortalRpcLocator.h"
 #include "PortalRpcMessages.h"
-#include "PortalRpcProvider.h"
+#include "PortalRpcResponder.h"
diff --git a/Engine/Source/Runtime/Portal/Rpc/Private/PortalRpcProvider.cpp b/Engine/Source/Runtime/Portal/Rpc/Private/PortalRpcProvider.cpp
deleted file mode 100644
diff --git a/Engine/Source/Runtime/Portal/Rpc/Private/PortalRpcProvider.h b/Engine/Source/Runtime/Portal/Rpc/Private/PortalRpcProvider.h
deleted file mode 100644
diff --git a/Engine/Source/Runtime/Portal/Rpc/Private/PortalRpcResponder.cpp b/Engine/Source/Runtime/Portal/Rpc/Private/PortalRpcResponder.cpp
new file mode 100644
--- /dev/null
+++ b/Engine/Source/Runtime/Portal/Rpc/Private/PortalRpcResponder.cpp
@@ -0,0 +1,40 @@
+// Copyright 1998-2015 Epic Games, Inc. All Rights Reserved.
+
+#include "PortalRpcPrivatePCH.h"
+#include "IMessageRpcServer.h"
+#include "ModuleManager.h"
+
+
+/* FPortalRpcResponder structors
+ *****************************************************************************/
+
+FPortalRpcResponder::FPortalRpcResponder()
+{
+	MessageEndpoint = FMessageEndpoint::Builder("FRpcManager")
+		.Handling<FPortalRpcLocateServer>(this, &FPortalRpcResponder::HandleMessage);
+}
+
+
+/* FPortalRpcResponder callbacks
+ *****************************************************************************/
+
+void FPortalRpcResponder::HandleMessage(const FPortalRpcLocateServer& Message, const TSharedRef<IMessageContext, ESPMode::ThreadSafe>& Context)
+{
+	if (!LookupDelegate.IsBound())
+	{
+		return;
+	}
+
+	const FString ProductKey = Message.ProductId.ToString() + Message.ProductVersion;
+	TSharedPtr<IMessageRpcServer> Server = Servers.FindRef(ProductKey);
+
+	if (!Server.IsValid())
+	{
+		Server = LookupDelegate.Execute(ProductKey);
+	}
+
+	if (Server.IsValid())
+	{
+		MessageEndpoint->Send(new FPortalRpcServer(Server->GetAddress().ToString()), Context->GetSender());
+	}
+}
diff --git a/Engine/Source/Runtime/Portal/Rpc/Private/PortalRpcResponder.h b/Engine/Source/Runtime/Portal/Rpc/Private/PortalRpcResponder.h
new file mode 100644
--- /dev/null
+++ b/Engine/Source/Runtime/Portal/Rpc/Private/PortalRpcResponder.h
@@ -0,0 +1,50 @@
+// Copyright 1998-2015 Epic Games, Inc. All Rights Reserved.
+
+#pragma once
+
+#include "IPortalRpcResponder.h"
+
+
+class FMessageEndpoint;
+struct FPortalRpcLocateServer;
+class IMessageContext;
+
+
+class FPortalRpcResponder
+	: public IPortalRpcResponder
+{
+public:
+
+	/** Default constructor. */
+	FPortalRpcResponder();
+
+	/** Virtual destructor. */
+	virtual ~FPortalRpcResponder() { }
+
+public:
+
+	// IPortalRpcResponder interface
+
+	virtual FOnPortalRpcLookup& OnLookup() override
+	{
+		return LookupDelegate;
+	}
+
+private:
+
+private:
+
+	/** Handles FPortalRpcLocateServer messages. */
+	void HandleMessage(const FPortalRpcLocateServer& Message, const TSharedRef<IMessageContext, ESPMode::ThreadSafe>& Context);
+
+private:
+
+	/** A delegate that is executed when a look-up for an RPC server occurs. */
+	FOnPortalRpcLookup LookupDelegate;
+
+	/** Message endpoint. */
+	TSharedPtr<FMessageEndpoint, ESPMode::ThreadSafe> MessageEndpoint;
+
+	/** Holds the existing RPC servers. */
+	TMap<FString, TSharedPtr<IMessageRpcServer>> Servers;
+};
diff --git a/Engine/Source/Runtime/Portal/Rpc/Public/IPortalRpcModule.h b/Engine/Source/Runtime/Portal/Rpc/Public/IPortalRpcModule.h
--- a/Engine/Source/Runtime/Portal/Rpc/Public/IPortalRpcModule.h
+++ b/Engine/Source/Runtime/Portal/Rpc/Public/IPortalRpcModule.h
@@ -6,7 +6,7 @@
 
 
 class IPortalRpcLocator;
-class IPortalRpcProvider;
+class IPortalRpcResponder;
 
 
 /**
@@ -18,5 +18,5 @@ class IPortalRpcModule
 public:
 
 	virtual TSharedRef<IPortalRpcLocator> CreateLocator() = 0;
-	virtual TSharedRef<IPortalRpcProvider> CreateProvider() = 0;
+	virtual TSharedRef<IPortalRpcResponder> CreateResponder() = 0;
 };
diff --git a/Engine/Source/Runtime/Portal/Rpc/Public/IPortalRpcProvider.h b/Engine/Source/Runtime/Portal/Rpc/Public/IPortalRpcProvider.h
deleted file mode 100644
diff --git a/Engine/Source/Runtime/Portal/Rpc/Public/IPortalRpcResponder.h b/Engine/Source/Runtime/Portal/Rpc/Public/IPortalRpcResponder.h
new file mode 100644
--- /dev/null
+++ b/Engine/Source/Runtime/Portal/Rpc/Public/IPortalRpcResponder.h
@@ -0,0 +1,26 @@
+// Copyright 1998-2015 Epic Games, Inc. All Rights Reserved.
+
+#pragma once
+
+
+class IMessageRpcServer;
+
+
+DECLARE_DELEGATE_RetVal_OneParam(TSharedPtr<IMessageRpcServer>, FOnPortalRpcLookup, const FString& /*ProductKey*/)
+
+
+/**
+ * Interface for Portal RPC responders.
+ */
+class IPortalRpcResponder
+{
+public:
+
+	/** Get a delegate that is executed when a look-up for an RPC server occurs. */
+	virtual FOnPortalRpcLookup& OnLookup() = 0;
+
+public:
+
+	/** Virtual destructor. */
+	~IPortalRpcResponder() { }
+};
