Portal: Started to implement service proxies

diff --git a/Engine/Source/Runtime/Portal/Proxies/PortalProxies.Build.cs b/Engine/Source/Runtime/Portal/Proxies/PortalProxies.Build.cs
new file mode 100644
--- /dev/null
+++ b/Engine/Source/Runtime/Portal/Proxies/PortalProxies.Build.cs
@@ -0,0 +1,49 @@
+// Copyright 1998-2015 Epic Games, Inc. All Rights Reserved.
+
+namespace UnrealBuildTool.Rules
+{
+	public class PortalProxies : ModuleRules
+	{
+		public PortalProxies(TargetInfo Target)
+		{
+            DynamicallyLoadedModuleNames.AddRange(
+                new string[] {
+                    "Messaging",
+                    "MessagingRpc",
+                }
+            );
+
+			PublicDependencyModuleNames.AddRange(
+				new string[]
+				{
+					"Core",
+				}
+			);
+			
+			PrivateDependencyModuleNames.AddRange(
+				new string[]
+				{
+					"CoreUObject",
+					"PortalMessages",
+				}
+			);
+
+            PrivateIncludePathModuleNames.AddRange(
+                new string[] {
+                    "Messaging",
+                    "MessagingRpc",
+                    "PortalServices",
+                }
+            );
+
+			PrivateIncludePaths.AddRange(
+				new string[] {
+					"Runtime/Portal/Proxies/Private",
+					"Runtime/Portal/Proxies/Private/Account",
+					"Runtime/Portal/Proxies/Private/Application",
+					"Runtime/Portal/Proxies/Private/Package",
+				}
+			);
+		}
+	}
+}
diff --git a/Engine/Source/Runtime/Portal/Proxies/Private/Application/PortalAppWindowProxy.h b/Engine/Source/Runtime/Portal/Proxies/Private/Application/PortalAppWindowProxy.h
new file mode 100644
--- /dev/null
+++ b/Engine/Source/Runtime/Portal/Proxies/Private/Application/PortalAppWindowProxy.h
@@ -0,0 +1,46 @@
+// Copyright 1998-2015 Epic Games, Inc. All Rights Reserved.
+
+#pragma once
+
+#include "IMessageRpcClient.h"
+#include "IPortalAppWindow.h"
+#include "PortalAppWindowMessages.h"
+
+
+class FPortalAppWindowProxy
+	: public IPortalAppWindow
+{
+public:
+
+	FPortalAppWindowProxy(const TSharedRef<IMessageRpcClient>& InRpcClient)
+		: RpcClient(InRpcClient)
+	{ }
+
+	virtual ~FPortalAppWindowProxy() { }
+
+public:
+
+	virtual TAsyncResult<bool> OpenFriends() override
+	{
+		return RpcClient->Call<FPortalAppWindowOpenFriends>();
+	}
+
+	virtual TAsyncResult<bool> OpenLibrary() override
+	{
+		return RpcClient->Call<FPortalAppWindowOpenLibrary>();
+	}
+
+	virtual TAsyncResult<bool> OpenMarketplace() override
+	{
+		return RpcClient->Call<FPortalAppWindowOpenMarketplace>();
+	}
+
+	virtual TAsyncResult<bool> OpenMarketItem(const FString& ItemId) override
+	{
+		return RpcClient->Call<FPortalAppWindowOpenMarketItem>(ItemId);
+	}
+
+private:
+
+	TSharedRef<IMessageRpcClient> RpcClient;
+};
diff --git a/Engine/Source/Runtime/Portal/Proxies/Private/Package/PortalPackageInstallerProxy.h b/Engine/Source/Runtime/Portal/Proxies/Private/Package/PortalPackageInstallerProxy.h
new file mode 100644
--- /dev/null
+++ b/Engine/Source/Runtime/Portal/Proxies/Private/Package/PortalPackageInstallerProxy.h
@@ -0,0 +1,36 @@
+// Copyright 1998-2015 Epic Games, Inc. All Rights Reserved.
+
+#pragma once
+
+#include "IPortalPackageInstaller.h"
+#include "IMessageRpcClient.h"
+#include "PortalPackageInstallerMessages.h"
+
+
+class FPortalPackageInstallerProxy
+	: public IPortalPackageInstaller
+{
+public:
+
+	FPortalPackageInstallerProxy(const TSharedRef<IMessageRpcClient>& InRpc)
+		: RpcClient(InRpc)
+	{ }
+
+	virtual ~FPortalPackageInstallerProxy() { }
+
+public:
+
+	virtual TAsyncResult<bool> Install(const FString& Path, const FString& AppName, const FString& BuildLabel) override
+	{
+		return RpcClient->Call<FPortalPackageInstallerInstall>(Path, AppName, BuildLabel);
+	}
+
+	virtual TAsyncResult<bool> Uninstall(const FString& Path, const FString& AppName, const FString& BuildLabel, bool RemoveUserFiles) override
+	{
+		return RpcClient->Call<FPortalPackageInstallerUninstall>(Path, AppName, BuildLabel, RemoveUserFiles);
+	}
+
+private:
+
+	TSharedRef<IMessageRpcClient> RpcClient;
+};
diff --git a/Engine/Source/Runtime/Portal/Proxies/Private/PortalProxiesModule.cpp b/Engine/Source/Runtime/Portal/Proxies/Private/PortalProxiesModule.cpp
new file mode 100644
--- /dev/null
+++ b/Engine/Source/Runtime/Portal/Proxies/Private/PortalProxiesModule.cpp
@@ -0,0 +1,54 @@
+// Copyright 1998-2015 Epic Games, Inc. All Rights Reserved.
+
+#include "PortalProxiesPrivatePCH.h"
+#include "IPortalServiceProvider.h"
+#include "IPortalAppWindow.h"
+
+
+/**
+ * Implements the PortalProxies module.
+ */
+class FPortalProxiesModule
+	: public IPortalServiceProvider
+{
+public:
+
+	// IModuleInterface interface
+
+	virtual void StartupModule() override { }
+	virtual void ShutdownModule() override { }
+
+	virtual bool SupportsDynamicReloading() override
+	{
+		return false;
+	}
+
+public:
+
+	// IPortalServiceProvider interface
+
+	virtual TSharedPtr<IPortalService> GetService(const FString& ServiceName, const TSharedRef<FTypeContainer>& Dependencies) override
+	{
+		TSharedPtr<IMessageRpcClient> RpcClient = Dependencies->GetInstance<IMessageRpcClient>();
+
+		if (!RpcClient.IsValid())
+		{
+			return nullptr;
+		}
+
+		if (ServiceName == TNameOf<IPortalAppWindow>::GetName())
+		{
+			return MakeShareable(new FPortalAppWindowProxy(RpcClient.ToSharedRef()));
+		}
+		
+		if (ServiceName == TNameOf<IPortalPackageInstaller>::GetName())
+		{
+			return MakeShareable(new FPortalPackageInstallerProxy(RpcClient.ToSharedRef()));
+		}
+
+		return nullptr;
+	}
+};
+
+
+IMPLEMENT_MODULE(FPortalProxiesModule, PortalProxies);
diff --git a/Engine/Source/Runtime/Portal/Proxies/Private/PortalProxiesPrivatePCH.h b/Engine/Source/Runtime/Portal/Proxies/Private/PortalProxiesPrivatePCH.h
new file mode 100644
--- /dev/null
+++ b/Engine/Source/Runtime/Portal/Proxies/Private/PortalProxiesPrivatePCH.h
@@ -0,0 +1,23 @@
+// Copyright 1998-2015 Epic Games, Inc. All Rights Reserved.
+
+#pragma once
+
+
+/* Private dependencies
+ *****************************************************************************/
+
+#include "Core.h"
+#include "CoreUObject.h"
+#include "AsyncResult.h"
+#include "IMessageRpcClient.h"
+#include "IMessageRpcServer.h"
+#include "MessageEndpoint.h"
+#include "MessageEndpointBuilder.h"
+#include "TypeContainer.h"
+
+
+/* Private includes
+ *****************************************************************************/
+
+#include "PortalAppWindowProxy.h"
+#include "PortalPackageInstallerProxy.h"
