Portal: Added missing message subscription to PortalRpcResponder

diff --git a/Engine/Source/Runtime/Portal/Rpc/Private/PortalRpcResponder.cpp b/Engine/Source/Runtime/Portal/Rpc/Private/PortalRpcResponder.cpp
--- a/Engine/Source/Runtime/Portal/Rpc/Private/PortalRpcResponder.cpp
+++ b/Engine/Source/Runtime/Portal/Rpc/Private/PortalRpcResponder.cpp
@@ -10,8 +10,13 @@
 
 FPortalRpcResponder::FPortalRpcResponder()
 {
-	MessageEndpoint = FMessageEndpoint::Builder("FRpcManager")
+	MessageEndpoint = FMessageEndpoint::Builder("FPortalRpcResponder")
 		.Handling<FPortalRpcLocateServer>(this, &FPortalRpcResponder::HandleMessage);
+
+	if (MessageEndpoint.IsValid())
+	{
+		MessageEndpoint->Subscribe<FPortalRpcLocateServer>();
+	}
 }
 
 
