Portal: Added missing forward declarations

diff --git a/Engine/Source/Runtime/Portal/Services/Public/IPortalServicesModule.h b/Engine/Source/Runtime/Portal/Services/Public/IPortalServicesModule.h
--- a/Engine/Source/Runtime/Portal/Services/Public/IPortalServicesModule.h
+++ b/Engine/Source/Runtime/Portal/Services/Public/IPortalServicesModule.h
@@ -5,6 +5,10 @@
 #include "ModuleInterface.h"
 
 
+class FTypeContainer;
+class IPortalServiceLocator;
+
+
 class IPortalServicesModule
 	: public IModuleInterface
 {
