Setup first complete Portal Service. I have not submitted the code which uses the service yet however.

Merging using UE4-To-UE4-LauncherDev

diff --git a/Engine/Source/Editor/UnrealEd/Classes/Editor/EditorEngine.h b/Engine/Source/Editor/UnrealEd/Classes/Editor/EditorEngine.h
--- a/Engine/Source/Editor/UnrealEd/Classes/Editor/EditorEngine.h
+++ b/Engine/Source/Editor/UnrealEd/Classes/Editor/EditorEngine.h
@@ -2683,6 +2683,9 @@ protected:
 	 */
 	void InvalidateAllViewportsAndHitProxies();
 
+	/** Initialize Portal RPC. */
+	void InitializePortal();
+
 	/** Destroy any online subsystems generated by PIE */
 	void CleanupPIEOnlineSessions(TArray<FName> OnlineIdentifiers);
 
diff --git a/Engine/Source/Editor/UnrealEd/UnrealEd.Build.cs b/Engine/Source/Editor/UnrealEd/UnrealEd.Build.cs
--- a/Engine/Source/Editor/UnrealEd/UnrealEd.Build.cs
+++ b/Engine/Source/Editor/UnrealEd/UnrealEd.Build.cs
@@ -57,6 +57,9 @@ public class UnrealEd : ModuleRules
                 "Niagara",
 				"SizeMap",
                 "LocalizationService",
+                "MessagingRpc",
+                "PortalRpc",
+                "PortalServices",
 			}
 		);
 
@@ -201,6 +204,8 @@ public class UnrealEd : ModuleRules
                 "IOSPlatformEditor",
 				"HTML5PlatformEditor",
 				"SizeMap",
+                "PortalProxies",
+                "PortalServices",
 			}
 		);
 
diff --git a/Engine/Source/Runtime/Engine/Classes/Engine/Engine.h b/Engine/Source/Runtime/Engine/Classes/Engine/Engine.h
--- a/Engine/Source/Runtime/Engine/Classes/Engine/Engine.h
+++ b/Engine/Source/Runtime/Engine/Classes/Engine/Engine.h
@@ -12,6 +12,12 @@ class FViewport;
 class FCommonViewportClient;
 class FCanvas;
 
+#if PLATFORM_COMPILER_HAS_VARIADIC_TEMPLATES
+class FTypeContainer;
+class IMessageRpcClient;
+class IPortalRpcLocator;
+class IPortalServiceLocator;
+#endif
 
 /**
  * Enumerates types of fully loaded packages.
@@ -1158,9 +1164,9 @@ public:
 	/** Whether to use a fixed framerate. */
 	UPROPERTY(config, EditAnywhere, Category = Framerate)
 	uint32 bUseFixedFrameRate : 1;
-
+	
 	/** The fixed framerate to use. */
-	UPROPERTY(config, EditAnywhere, Category = Framerate, meta = (EditCondition = "bUseFixedFrameRate"))
+	UPROPERTY(config, EditAnywhere, Category = Framerate, meta=(EditCondition="bUseFixedFrameRate"))
 	float FixedFrameRate;
 
 	/** Range of framerates in which smoothing will kick in */
@@ -1499,7 +1505,7 @@ public:
 
 	/** Reference to the HMD device that is attached, if any */
 	TSharedPtr< class IHeadMountedDisplay, ESPMode::ThreadSafe > HMDDevice;
-
+	
 	/** Reference to the Motion Control devices that are attached, if any */
 	TArray < class IMotionController*> MotionControllerDevices;
 
@@ -2153,6 +2159,33 @@ public:
 
 	virtual void RemapGamepadControllerIdForPIE(class UGameViewportClient* InGameViewport, int32 &ControllerId) { }
 
+#if PLATFORM_COMPILER_HAS_VARIADIC_TEMPLATES
+	/**
+	 * Get a locator for Portal services.
+	 *
+	 * @return The service locator.
+	 */
+	TSharedRef<IPortalServiceLocator> GetServiceLocator()
+	{
+		return ServiceLocator.ToSharedRef();
+	}
+
+protected:
+
+	/** Portal RPC client. */
+	TSharedPtr<IMessageRpcClient> PortalRpcClient;
+
+	/** Portal RPC server locator. */
+	TSharedPtr<IPortalRpcLocator> PortalRpcLocator;
+
+	/** Holds a type container for service dependencies. */
+	TSharedPtr<FTypeContainer> ServiceDependencies;
+
+	/** Holds registered service instances. */
+	TSharedPtr<IPortalServiceLocator> ServiceLocator;
+#endif
+
+
 public:
 
 	/**
@@ -2239,26 +2272,20 @@ protected:
 	 */
 	virtual bool InitializeHMDDevice();
 
-	/**
-	 *	Detects and initializes any motion controller devices
-	 *
-	 *	@return true if there are any initialized motion controllers, false otherwise
-	 */
 	virtual bool InitializeMotionControllers();
 
-	/**
-	 *	Record EngineAnalytics information for attached HMD devices
-	 */
+	/**	Record EngineAnalytics information for attached HMD devices. */
 	virtual void RecordHMDAnalytics();
 
-	/**
-	 * Loads all Engine object references from their corresponding config entries.
-	 */
+	/** Loads all Engine object references from their corresponding config entries. */
 	virtual void InitializeObjectReferences();
 
-	/** 
-	 * Initializes the running average delta to some good initial framerate 
-	 */
+#if PLATFORM_COMPILER_HAS_VARIADIC_TEMPLATES
+	/** Initialize Portal services. */
+	virtual void InitializePortalServices();
+#endif
+
+	/** Initializes the running average delta to some good initial framerate. */
 	virtual void InitializeRunningAverageDeltaTime();
 
 	float RunningAverageDeltaTime;
diff --git a/Engine/Source/Runtime/Engine/Engine.Build.cs b/Engine/Source/Runtime/Engine/Engine.Build.cs
--- a/Engine/Source/Runtime/Engine/Engine.Build.cs
+++ b/Engine/Source/Runtime/Engine/Engine.Build.cs
@@ -75,6 +75,18 @@ public class Engine : ModuleRules
 			}
         );
 
+        if (Target.Platform != UnrealTargetPlatform.XboxOne)
+        {
+            // these modules require variadic templates
+            PrivateDependencyModuleNames.AddRange(
+                new string[] {
+                    "MessagingRpc",
+                    "PortalRpc",
+                    "PortalServices",
+                }
+            );
+        }
+
         CircularlyReferencedDependentModules.Add("AIModule");
 		CircularlyReferencedDependentModules.Add("Landscape");
         CircularlyReferencedDependentModules.Add("UMG");
diff --git a/Engine/Source/Runtime/Engine/Private/UnrealEngine.cpp b/Engine/Source/Runtime/Engine/Private/UnrealEngine.cpp
--- a/Engine/Source/Runtime/Engine/Private/UnrealEngine.cpp
+++ b/Engine/Source/Runtime/Engine/Private/UnrealEngine.cpp
@@ -105,7 +105,17 @@
 #include "JsonInternationalizationArchiveSerializer.h"
 #include "JsonInternationalizationManifestSerializer.h"
 
+#if PLATFORM_COMPILER_HAS_VARIADIC_TEMPLATES
+#include "IMessagingRpcModule.h"
+#include "IMessageRpcClient.h"
+#include "IPortalRpcModule.h"
+#include "IPortalRpcLocator.h"
+#include "IPortalServicesModule.h"
+#include "TypeContainer.h"
+#endif
+
 #include "GameFramework/OnlineSession.h"
+#include "IPortalServiceLocator.h"
 
 #include "InstancedReferenceSubobjectHelper.h"
 
@@ -902,6 +912,14 @@ void UEngine::Init(IEngineLoop* InEngineLoop)
 	// Initialise buffer visualization system data
 	GetBufferVisualizationData().Initialize();
 
+#if PLATFORM_COMPILER_HAS_VARIADIC_TEMPLATES
+	// Initialize Portal services
+	if (!IsRunningCommandlet())
+	{
+		InitializePortalServices();
+	}
+#endif
+
 	// Connect the engine analytics provider
 	FEngineAnalytics::Initialize();
 
@@ -1876,6 +1894,34 @@ void UEngine::InitializeObjectReferences()
 	UISettings->ForceLoadResources();
 }
 
+#if PLATFORM_COMPILER_HAS_VARIADIC_TEMPLATES
+void UEngine::InitializePortalServices()
+{
+	// Initialize Portal services
+	PortalRpcClient = FModuleManager::LoadModuleChecked<IMessagingRpcModule>("MessagingRpc").CreateRpcClient();
+	{
+		// @todo gmp: catch timeouts?
+	}
+
+	PortalRpcLocator = FModuleManager::LoadModuleChecked<IPortalRpcModule>("PortalRpc").CreateLocator();
+	{
+		PortalRpcLocator->OnServerLocated().BindLambda([=]() { PortalRpcClient->Connect(PortalRpcLocator->GetServerAddress()); });
+		PortalRpcLocator->OnServerLost().BindLambda([=]() { PortalRpcClient->Disconnect(); });
+	}
+
+	ServiceDependencies = MakeShareable(new FTypeContainer);
+	{
+		ServiceDependencies->RegisterInstance<IMessageRpcClient>(PortalRpcClient.ToSharedRef());
+	}
+
+	ServiceLocator = FModuleManager::LoadModuleChecked<IPortalServicesModule>("PortalServices").CreateLocator(ServiceDependencies.ToSharedRef());
+	{
+		// @todo add any Engine specific Portal services here
+		ServiceLocator->Configure(TEXT("IPortalApplicationWindow"), TEXT("*"), "PortalProxies");
+	}
+}
+#endif
+
 //
 // Exit the engine.
 //
diff --git a/Engine/Source/Runtime/MessagingRpc/Private/MessageRpcClient.cpp b/Engine/Source/Runtime/MessagingRpc/Private/MessageRpcClient.cpp
--- a/Engine/Source/Runtime/MessagingRpc/Private/MessageRpcClient.cpp
+++ b/Engine/Source/Runtime/MessagingRpc/Private/MessageRpcClient.cpp
@@ -8,8 +8,9 @@
 
 FMessageRpcClient::FMessageRpcClient()
 {
-	MessageEndpoint = FMessageEndpoint::Builder("FPortalAppWindowEndpoint")
-		.Handling<FMessageRpcProgress>(this, &FMessageRpcClient::HandleMessage);
+	MessageEndpoint = FMessageEndpoint::Builder("FMessageRpcClient")
+		.Handling<FMessageRpcProgress>(this, &FMessageRpcClient::HandleProgressMessage)
+		.WithCatchall(this, &FMessageRpcClient::HandleMessage);
 
 	TickerHandle = FTicker::GetCoreTicker().AddTicker(FTickerDelegate::CreateRaw(this, &FMessageRpcClient::HandleTicker), MESSAGE_RPC_RETRY_INTERVAL);
 }
@@ -24,6 +25,12 @@ FMessageRpcClient::~FMessageRpcClient()
 /* IMessageRpcClient interface
  *****************************************************************************/
 
+bool FMessageRpcClient::IsConnected() const
+{
+	return ServerAddress.IsValid();
+}
+
+
 void FMessageRpcClient::Connect(const FMessageAddress& InServerAddress)
 {
 	Disconnect();
@@ -58,7 +65,7 @@ void FMessageRpcClient::SendCall(const TSharedPtr<IMessageRpcCall>& Call)
 	if (ServerAddress.IsValid())
 	{
 		MessageEndpoint->Send(
-			Call->GetMessage(),
+			Call->ConstructMessage(),
 			Call->GetMessageType(),
 			nullptr,
 			TArrayBuilder<FMessageAddress>().Add(ServerAddress),
@@ -93,16 +100,45 @@ void FMessageRpcClient::CancelCall(const FGuid& CallId)
 /* FMessageRpcClient event handlers
  *****************************************************************************/
 
-void FMessageRpcClient::HandleMessage(const FMessageRpcProgress& Message, const TSharedRef<IMessageContext, ESPMode::ThreadSafe>& Context)
+void FMessageRpcClient::HandleProgressMessage(const FMessageRpcProgress& Message, const TSharedRef<IMessageContext, ESPMode::ThreadSafe>& Context)
 {
-	TSharedPtr<IMessageRpcCall> Call = FindCall(Context);
+	const TWeakObjectPtr<UScriptStruct>& MessageTypeInfo = Context->GetMessageTypeInfo();
+
+	if (!MessageTypeInfo.IsValid())
+	{
+		return;
+	}
 
-	if (Call.IsValid())
+	if (MessageTypeInfo->IsChildOf(FRpcMessage::StaticStruct()))
 	{
-		Call->UpdateProgress(Message.Completion, FText::FromString(Message.StatusText));
+		TSharedPtr<IMessageRpcCall> Call = FindCall(Context);
+
+		if (Call.IsValid())
+		{
+			Call->UpdateProgress(Message.Completion, FText::FromString(Message.StatusText));
+		}
 	}
 }
 
+void FMessageRpcClient::HandleMessage(const TSharedRef<IMessageContext, ESPMode::ThreadSafe>& Context)
+{
+	const TWeakObjectPtr<UScriptStruct>& MessageTypeInfo = Context->GetMessageTypeInfo();
+
+	if (!MessageTypeInfo.IsValid())
+	{
+		return;
+	}
+
+	if (MessageTypeInfo->IsChildOf(FRpcMessage::StaticStruct()))
+	{
+		TSharedPtr<IMessageRpcCall> Call;
+		auto Request = static_cast<const FRpcMessage*>(Context->GetMessage());
+		if (Calls.RemoveAndCopyValue(Request->CallId, Call))
+		{
+			Call->Complete(Context);
+		}
+	}
+}
 
 bool FMessageRpcClient::HandleTicker(float DeltaTime)
 {
diff --git a/Engine/Source/Runtime/MessagingRpc/Private/MessageRpcClient.h b/Engine/Source/Runtime/MessagingRpc/Private/MessageRpcClient.h
--- a/Engine/Source/Runtime/MessagingRpc/Private/MessageRpcClient.h
+++ b/Engine/Source/Runtime/MessagingRpc/Private/MessageRpcClient.h
@@ -30,6 +30,7 @@ public:
 
 	// IMessageRpcClient interface
 
+	virtual bool IsConnected() const override;
 	virtual void Connect(const FMessageAddress& InServerAddress) override;
 	virtual void Disconnect() override;
 
@@ -60,7 +61,9 @@ protected:
 private:
 
 	/** Handles FMessageRpcProgress messages. */
-	void HandleMessage(const FMessageRpcProgress& Message, const TSharedRef<IMessageContext, ESPMode::ThreadSafe>& Context);
+	void HandleProgressMessage(const FMessageRpcProgress& Message, const TSharedRef<IMessageContext, ESPMode::ThreadSafe>& Context);
+
+	void HandleMessage(const TSharedRef<IMessageContext, ESPMode::ThreadSafe>& Context);
 
 	/** Handles the ticker. */
 	bool HandleTicker(float DeltaTime);
diff --git a/Engine/Source/Runtime/MessagingRpc/Private/MessageRpcMessages.h b/Engine/Source/Runtime/MessagingRpc/Private/MessageRpcMessages.h
deleted file mode 100644
diff --git a/Engine/Source/Runtime/MessagingRpc/Private/MessageRpcServer.cpp b/Engine/Source/Runtime/MessagingRpc/Private/MessageRpcServer.cpp
--- a/Engine/Source/Runtime/MessagingRpc/Private/MessageRpcServer.cpp
+++ b/Engine/Source/Runtime/MessagingRpc/Private/MessageRpcServer.cpp
@@ -1,16 +1,18 @@
 // Copyright 1998-2015 Epic Games, Inc. All Rights Reserved.
 
 #include "PrivatePCH.h"
+#include "MessageRpcServer.h"
+#include "IMessageRpcServer.h"
+#include "IMessageRpcReturn.h"
 #include "IMessageRpcHandler.h"
 #include "IMessageRpcReturn.h"
 
-
 /* FMessageRpcServer structors
  *****************************************************************************/
 
 FMessageRpcServer::FMessageRpcServer()
 {
-	MessageEndpoint = FMessageEndpoint::Builder("FPortalAppWindowEndpoint")
+	MessageEndpoint = FMessageEndpoint::Builder("FMessageRpcServer")
 		.WithCatchall(this, &FMessageRpcServer::HandleMessage);
 
 	TickerHandle = FTicker::GetCoreTicker().AddTicker(FTickerDelegate::CreateRaw(this, &FMessageRpcServer::HandleTicker), 0.1f);
@@ -26,6 +28,11 @@ FMessageRpcServer::~FMessageRpcServer()
 /* IMessageRpcServer interface
  *****************************************************************************/
 
+TSharedPtr<FMessageEndpoint, ESPMode::ThreadSafe> FMessageRpcServer::GetEndpoint() const
+{
+	return MessageEndpoint;
+}
+
 void FMessageRpcServer::AddHandler(const FName& RequestMessageType, const TSharedRef<IMessageRpcHandler>& Handler)
 {
 	Handlers.Add(RequestMessageType, Handler);
@@ -113,6 +120,7 @@ void FMessageRpcServer::SendProgress(const FGuid& CallId, const FReturnInfo& Ret
 void FMessageRpcServer::SendResult(const FGuid& CallId, const FReturnInfo& ReturnInfo)
 {
 	FRpcMessage* Message = ReturnInfo.Return->CreateResponseMessage();
+	Message->CallId = CallId;
 
 	MessageEndpoint->Send(
 		Message,
diff --git a/Engine/Source/Runtime/MessagingRpc/Private/MessageRpcServer.h b/Engine/Source/Runtime/MessagingRpc/Private/MessageRpcServer.h
deleted file mode 100644
diff --git a/Engine/Source/Runtime/MessagingRpc/Public/IMessageRpcCall.h b/Engine/Source/Runtime/MessagingRpc/Public/IMessageRpcCall.h
--- a/Engine/Source/Runtime/MessagingRpc/Public/IMessageRpcCall.h
+++ b/Engine/Source/Runtime/MessagingRpc/Public/IMessageRpcCall.h
@@ -35,11 +35,18 @@ public:
 	virtual const FGuid& GetId() const = 0;
 
 	/**
-	 * Get the request message.
+	 * Get the request message template.
+	 *
+	 * @return Pointer to the request message template.
+	 */
+	virtual void* GetMessageTemplate() const = 0;
+
+	/**
+	 * Constructs a new message based on the call message template. Ownership of the message belongs to the caller.
 	 *
 	 * @return Pointer to the request message.
 	 */
-	virtual void* GetMessage() const = 0;
+	virtual void* ConstructMessage() const = 0;
 
 	/**
 	 * Get the time at which the request was last updated by the server.
diff --git a/Engine/Source/Runtime/MessagingRpc/Public/IMessageRpcClient.h b/Engine/Source/Runtime/MessagingRpc/Public/IMessageRpcClient.h
--- a/Engine/Source/Runtime/MessagingRpc/Public/IMessageRpcClient.h
+++ b/Engine/Source/Runtime/MessagingRpc/Public/IMessageRpcClient.h
@@ -108,14 +108,14 @@ class IMessageRpcClient
 	public:
 
 		TCall(const FGuid& CallId, P... Params)
-			: Message(new typename RpcType::FRequest(Params...))
+			: MessageTemplate(new typename RpcType::FRequest(Params...))
 		{
-			Message->CallId = CallId;
+			MessageTemplate->CallId = CallId;
 		}
 
 		virtual ~TCall()
 		{
-			delete Message;
+			delete MessageTemplate;
 		}
 
 		virtual void Complete(const TSharedRef<IMessageContext, ESPMode::ThreadSafe>& ResponseContext) override
@@ -139,12 +139,17 @@ class IMessageRpcClient
 
 		virtual const FGuid& GetId() const override
 		{
-			return Message->CallId;
+			return MessageTemplate->CallId;
 		}
 
-		virtual void* GetMessage() const override
+		virtual void* GetMessageTemplate() const override
 		{
-			return Message;
+			return MessageTemplate;
+		}
+
+		virtual void* ConstructMessage() const override
+		{
+			return new typename RpcType::FRequest(*MessageTemplate);
 		}
 
 		virtual UScriptStruct* GetMessageType() const override
@@ -161,11 +166,16 @@ class IMessageRpcClient
 	private:
 
 		TPromise<typename RpcType::FResult> Promise;
-		typename RpcType::FRequest* Message;
+		typename RpcType::FRequest* MessageTemplate;
 	};
 
 public:
 
+	/**
+	 * @return Whether the MessageRpcClient is currently connected to an RPC server.
+	 */
+	virtual bool IsConnected() const = 0;
+
 	/**
 	 * Connect this client to an RPC server.
 	 *
diff --git a/Engine/Source/Runtime/MessagingRpc/Public/MessageRpcMessages.h b/Engine/Source/Runtime/MessagingRpc/Public/MessageRpcMessages.h
new file mode 100644
--- /dev/null
+++ b/Engine/Source/Runtime/MessagingRpc/Public/MessageRpcMessages.h
@@ -0,0 +1,75 @@
+// Copyright 1998-2015 Epic Games, Inc. All Rights Reserved.
+
+#pragma once
+
+#include "MessageRpcMessages.generated.h"
+
+
+/** Message for canceling an RPC call. */
+USTRUCT()
+struct FMessageRpcCancel
+{
+	GENERATED_USTRUCT_BODY()
+
+	/** Correlation identifier for the RPC call that this message refers to. */
+	UPROPERTY()
+	FGuid CallId;
+
+	/** Default constructor. */
+	FMessageRpcCancel() { }
+
+	/** Creates and initializes a new instance. */
+	FMessageRpcCancel(const FGuid& InCallId)
+		: CallId(InCallId)
+	{ }
+};
+
+
+/** Message for updating the progress of an RPC call. */
+USTRUCT()
+struct FMessageRpcProgress
+{
+	GENERATED_USTRUCT_BODY()
+
+	/** Completion percentage (0.0 to 1.0). */
+	UPROPERTY()
+	float Completion;
+
+	/** Correlation identifier for the RPC call that this message refers to. */
+	UPROPERTY()
+	FGuid CallId;
+
+	/** Status text. */
+	UPROPERTY()
+	FString StatusText;
+
+	/** Default constructor. */
+	FMessageRpcProgress() { }
+
+	/** Creates and initializes a new instance. */
+	FMessageRpcProgress(const FGuid& InCallId, float InCompletion, const FText& InStatusText)
+		: Completion(InCompletion)
+		, CallId(InCallId)
+		, StatusText(InStatusText.ToString())
+	{ }
+};
+
+
+/** Message for notifying RPC clients that a call was not handled by the server. */
+USTRUCT()
+struct FMessageRpcUnhandled
+{
+	GENERATED_USTRUCT_BODY()
+
+	/** Correlation identifier for the RPC call that this message refers to. */
+	UPROPERTY()
+	FGuid CallId;
+
+	/** Default constructor. */
+	FMessageRpcUnhandled() { }
+
+	/** Creates and initializes a new instance. */
+	FMessageRpcUnhandled(const FGuid& InCallId)
+		: CallId(InCallId)
+	{ }
+};
diff --git a/Engine/Source/Runtime/MessagingRpc/Public/MessageRpcServer.h b/Engine/Source/Runtime/MessagingRpc/Public/MessageRpcServer.h
new file mode 100644
--- /dev/null
+++ b/Engine/Source/Runtime/MessagingRpc/Public/MessageRpcServer.h
@@ -0,0 +1,87 @@
+// Copyright 1998-2015 Epic Games, Inc. All Rights Reserved.
+#pragma once
+
+#include "IMessageRpcServer.h"
+#include "MessageRpcMessages.h"
+
+class IMessageRpcHandler;
+class IMessageRpcReturn;
+class IMessageContext;
+class FMessageEndpoint;
+class FMessagingRpcModule;
+
+/**
+ * Implements an RPC server.
+ */
+class MESSAGINGRPC_API FMessageRpcServer
+	: public IMessageRpcServer
+{
+public:
+
+	/** Virtual destructor. */
+	virtual ~FMessageRpcServer();
+
+public:
+
+	// IMessageRpcServer interface
+	virtual void AddHandler(const FName& RequestMessageType, const TSharedRef<IMessageRpcHandler>& Handler) override;
+	virtual const FMessageAddress& GetAddress() const override;
+	virtual FOnMessageRpcNoHandler& OnNoHandler() override;
+
+protected:
+
+	/** Default constructor. */
+	FMessageRpcServer();
+
+	struct FReturnInfo
+	{
+		FMessageAddress ClientAddress;
+		FDateTime LastProgressSent;
+		TSharedPtr<IAsyncProgress> Progress;
+		TSharedPtr<IMessageRpcReturn> Return;
+		TSharedPtr<IAsyncTask> Task;
+	};
+
+	TSharedPtr<FMessageEndpoint, ESPMode::ThreadSafe> GetEndpoint() const;
+
+	/** Processes an FMessageRpcCancel message. */
+	void ProcessCancelation(const FMessageRpcCancel& Message, const TSharedRef<IMessageContext, ESPMode::ThreadSafe>& Context);
+
+	/** Processes an RPC request message. */
+	void ProcessRequest(const TSharedRef<IMessageContext, ESPMode::ThreadSafe>& Context);
+
+	/** Send a progress message to the RPC client that made the RPC call. */
+	void SendProgress(const FGuid& CallId, const FReturnInfo& ReturnInfo);
+
+	/** Send a result message to the RPC client that made the RPC call. */
+	void SendResult(const FGuid& CallId, const FReturnInfo& ReturnInfo);
+
+protected:
+
+	/** Message endpoint. */
+	TSharedPtr<FMessageEndpoint, ESPMode::ThreadSafe> MessageEndpoint;
+
+private:
+
+	/** Handles all incoming messages. */
+	void HandleMessage(const TSharedRef<IMessageContext, ESPMode::ThreadSafe>& Context);
+
+	/** Handles the ticker. */
+	bool HandleTicker(float DeltaTime);
+
+private:
+
+	/** Registered request message handlers. */
+	TMap<FName, TSharedPtr<IMessageRpcHandler>> Handlers;
+
+	/* Delegate that is executed when a received RPC message has no registered handler. */
+	FOnMessageRpcNoHandler NoHandlerDelegate;
+
+	/** Collection of pending RPC returns. */
+	TMap<FGuid, FReturnInfo> Returns;
+
+	/** Handle to the registered ticker. */
+	FDelegateHandle TickerHandle;
+
+	friend FMessagingRpcModule;
+};
diff --git a/Engine/Source/Runtime/Portal/Messages/Public/PortalAppWindowMessages.h b/Engine/Source/Runtime/Portal/Messages/Public/PortalAppWindowMessages.h
deleted file mode 100644
diff --git a/Engine/Source/Runtime/Portal/Messages/Public/PortalApplicationWindowMessages.h b/Engine/Source/Runtime/Portal/Messages/Public/PortalApplicationWindowMessages.h
new file mode 100644
--- /dev/null
+++ b/Engine/Source/Runtime/Portal/Messages/Public/PortalApplicationWindowMessages.h
@@ -0,0 +1,38 @@
+// Copyright 1998-2015 Epic Games, Inc. All Rights Reserved.
+#pragma once
+
+#include "RpcMessage.h"
+#include "PortalApplicationWindowMessages.generated.h"
+
+USTRUCT()
+struct FPortalApplicationWindowNavigateToRequest
+	: public FRpcMessage
+{
+	GENERATED_USTRUCT_BODY()
+
+	UPROPERTY()
+	FString Url;
+
+	FPortalApplicationWindowNavigateToRequest() { }
+	FPortalApplicationWindowNavigateToRequest(const FString& InUrl)
+		: Url(InUrl)
+	{ }
+};
+
+
+USTRUCT()
+struct FPortalApplicationWindowNavigateToResponse
+	: public FRpcMessage
+{
+	GENERATED_USTRUCT_BODY()
+
+	UPROPERTY()
+	bool Result;
+
+	FPortalApplicationWindowNavigateToResponse() { }
+	FPortalApplicationWindowNavigateToResponse(bool InResult)
+		: Result(InResult)
+	{ }
+};
+
+DECLARE_RPC(FPortalApplicationWindowNavigateTo, bool)
diff --git a/Engine/Source/Runtime/Portal/Proxies/Private/Application/PortalAppWindowProxy.h b/Engine/Source/Runtime/Portal/Proxies/Private/Application/PortalAppWindowProxy.h
deleted file mode 100644
diff --git a/Engine/Source/Runtime/Portal/Proxies/Private/Application/PortalApplicationWindowProxy.cpp b/Engine/Source/Runtime/Portal/Proxies/Private/Application/PortalApplicationWindowProxy.cpp
new file mode 100644
--- /dev/null
+++ b/Engine/Source/Runtime/Portal/Proxies/Private/Application/PortalApplicationWindowProxy.cpp
@@ -0,0 +1,44 @@
+// Copyright 1998-2015 Epic Games, Inc. All Rights Reserved.
+
+#include "PortalProxiesPrivatePCH.h"
+#include "IMessageRpcClient.h"
+#include "IPortalApplicationWindow.h"
+#include "PortalApplicationWindowMessages.h"
+
+class FPortalApplicationWindowProxy
+	: public IPortalApplicationWindow
+{
+public:
+
+	virtual ~FPortalApplicationWindowProxy() { }
+
+	virtual bool IsAvailable() const override
+	{
+		return RpcClient->IsConnected();
+	}
+
+public:
+
+	virtual TAsyncResult<bool> NavigateTo(const FString& Url) override
+	{
+		return RpcClient->Call<FPortalApplicationWindowNavigateTo>(Url);
+	}
+
+private:
+
+	FPortalApplicationWindowProxy(const TSharedRef<IMessageRpcClient>& InRpcClient)
+		: RpcClient(InRpcClient)
+	{ }
+
+
+private:
+
+	TSharedRef<IMessageRpcClient> RpcClient;
+
+	friend FPortalApplicationWindowProxyFactory;
+};
+
+TSharedRef<IPortalService> FPortalApplicationWindowProxyFactory::Create(const TSharedRef<IMessageRpcClient>& RpcClient)
+{
+	return MakeShareable(new FPortalApplicationWindowProxy(RpcClient));
+}
diff --git a/Engine/Source/Runtime/Portal/Proxies/Private/Application/PortalApplicationWindowProxy.h b/Engine/Source/Runtime/Portal/Proxies/Private/Application/PortalApplicationWindowProxy.h
new file mode 100644
--- /dev/null
+++ b/Engine/Source/Runtime/Portal/Proxies/Private/Application/PortalApplicationWindowProxy.h
@@ -0,0 +1,11 @@
+// Copyright 1998-2015 Epic Games, Inc. All Rights Reserved.
+#pragma once
+
+class IPortalService;
+class IMessageRpcClient;
+
+class FPortalApplicationWindowProxyFactory
+{
+public:
+	static TSharedRef<IPortalService> Create(const TSharedRef<IMessageRpcClient>& RpcClient);
+}; 
\ No newline at end of file
diff --git a/Engine/Source/Runtime/Portal/Proxies/Private/PortalProxiesModule.cpp b/Engine/Source/Runtime/Portal/Proxies/Private/PortalProxiesModule.cpp
--- a/Engine/Source/Runtime/Portal/Proxies/Private/PortalProxiesModule.cpp
+++ b/Engine/Source/Runtime/Portal/Proxies/Private/PortalProxiesModule.cpp
@@ -2,7 +2,7 @@
 
 #include "PortalProxiesPrivatePCH.h"
 #include "IPortalServiceProvider.h"
-#include "IPortalAppWindow.h"
+#include "IPortalApplicationWindow.h"
 
 
 /**
@@ -15,7 +15,10 @@ public:
 
 	// IModuleInterface interface
 
-	virtual void StartupModule() override { }
+	virtual void StartupModule() override 
+	{
+		ApplicationWindow_SupportedServiceNames.Add(TNameOf<IPortalApplicationWindow>::GetName());
+	}
 	virtual void ShutdownModule() override { }
 
 	virtual bool SupportsDynamicReloading() override
@@ -36,18 +39,20 @@ public:
 			return nullptr;
 		}
 
-		if (ServiceName == TNameOf<IPortalAppWindow>::GetName())
-		{
-			return MakeShareable(new FPortalAppWindowProxy(RpcClient.ToSharedRef()));
-		}
-		
-		if (ServiceName == TNameOf<IPortalPackageInstaller>::GetName())
+		if (ApplicationWindow_SupportedServiceNames.Contains(ServiceName))
 		{
-			return MakeShareable(new FPortalPackageInstallerProxy(RpcClient.ToSharedRef()));
+			return FPortalApplicationWindowProxyFactory::Create(RpcClient.ToSharedRef());
 		}
+	
+		//Add additional supported proxy services here
 
 		return nullptr;
 	}
+
+private:
+
+	TSet<FString> ApplicationWindow_SupportedServiceNames;
+
 };
 
 
diff --git a/Engine/Source/Runtime/Portal/Proxies/Private/PortalProxiesPrivatePCH.h b/Engine/Source/Runtime/Portal/Proxies/Private/PortalProxiesPrivatePCH.h
--- a/Engine/Source/Runtime/Portal/Proxies/Private/PortalProxiesPrivatePCH.h
+++ b/Engine/Source/Runtime/Portal/Proxies/Private/PortalProxiesPrivatePCH.h
@@ -19,5 +19,5 @@
 /* Private includes
  *****************************************************************************/
 
-#include "PortalAppWindowProxy.h"
+#include "PortalApplicationWindowProxy.h"
 #include "PortalPackageInstallerProxy.h"
diff --git a/Engine/Source/Runtime/Portal/Rpc/PortalRpc.Build.cs b/Engine/Source/Runtime/Portal/Rpc/PortalRpc.Build.cs
--- a/Engine/Source/Runtime/Portal/Rpc/PortalRpc.Build.cs
+++ b/Engine/Source/Runtime/Portal/Rpc/PortalRpc.Build.cs
@@ -8,8 +8,6 @@ namespace UnrealBuildTool.Rules
 		{
             DynamicallyLoadedModuleNames.AddRange(
                 new string[] {
-                    "Messaging",
-                    "MessagingRpc",
                 }
             );
 
@@ -18,6 +16,7 @@ namespace UnrealBuildTool.Rules
 				{
 					"Core",
 					"CoreUObject",
+                    "MessagingRpc",
 				}
 			);
 
diff --git a/Engine/Source/Runtime/Portal/Rpc/Private/PortalRpcLocator.cpp b/Engine/Source/Runtime/Portal/Rpc/Private/PortalRpcLocator.cpp
--- a/Engine/Source/Runtime/Portal/Rpc/Private/PortalRpcLocator.cpp
+++ b/Engine/Source/Runtime/Portal/Rpc/Private/PortalRpcLocator.cpp
@@ -1,77 +1,106 @@
 // Copyright 1998-2015 Epic Games, Inc. All Rights Reserved.
 
 #include "PortalRpcPrivatePCH.h"
-#include "EngineVersion.h"
 #include "PortalRpcLocator.h"
+#include "IPortalRpcLocator.h"
+#include "PortalRpcLocator.h"
+#include "PortalRpcMessages.h"
+#include "IMessageContext.h"
 
+#include "EngineVersion.h"
 
-/* FPortalRpcLocator structors
- *****************************************************************************/
-
-FPortalRpcLocator::FPortalRpcLocator()
-	: LastServerResponse(FDateTime::MinValue())
+class FPortalRpcLocatorImpl
+	: public IPortalRpcLocator
 {
-	MessageEndpoint = FMessageEndpoint::Builder("FPortalRpcLocator")
-		.Handling<FPortalRpcServer>(this, &FPortalRpcLocator::HandleMessage);
-
-	TickerHandle = FTicker::GetCoreTicker().AddTicker(FTickerDelegate::CreateRaw(this, &FPortalRpcLocator::HandleTicker), PORTAL_RPC_LOCATE_INTERVAL);
-}
+public:
 
+	virtual ~FPortalRpcLocatorImpl()
+	{
+		FTicker::GetCoreTicker().RemoveTicker(TickerHandle);
+	}
 
-FPortalRpcLocator::~FPortalRpcLocator()
-{
-	FTicker::GetCoreTicker().RemoveTicker(TickerHandle);
-}
+public:
 
+	// IPortalRpcLocator interface
 
-/* IPortalRpcLocator interface
- *****************************************************************************/
+	virtual const FMessageAddress& GetServerAddress() const override
+	{
+		return ServerAddress;
+	}
 
-const FMessageAddress& FPortalRpcLocator::GetServerAddress() const
-{
-	return ServerAddress;
-}
+	virtual FSimpleDelegate& OnServerLocated() override
+	{
+		return ServerLocatedDelegate;
+	}
 
+	virtual FSimpleDelegate& OnServerLost() override
+	{
+		return ServerLostDelegate;
+	}
 
-FSimpleDelegate& FPortalRpcLocator::OnServerLocated()
-{
-	return ServerLocatedDelegate;
-}
+private:
 
+	void HandleMessage(const FPortalRpcServer& Message, const TSharedRef<IMessageContext, ESPMode::ThreadSafe>& Context)
+	{
+		LastServerResponse = FDateTime::UtcNow();
 
-FSimpleDelegate& FPortalRpcLocator::OnServerLost()
-{
-	return ServerLostDelegate;
-}
+		FMessageAddress NewServerAddress;
 
+		if (FMessageAddress::Parse(Message.ServerAddress, NewServerAddress) && (NewServerAddress != ServerAddress))
+		{
+			ServerAddress = NewServerAddress;
+			ServerLocatedDelegate.ExecuteIfBound();
+		}
+	}
 
-/* FPortalRpcLocator callbacks
- *****************************************************************************/
+	bool HandleTicker(float DeltaTime)
+	{
+		if (ServerAddress.IsValid() && ((FDateTime::UtcNow() - LastServerResponse).GetSeconds() > PORTAL_RPC_LOCATE_TIMEOUT))
+		{
+			ServerAddress.Invalidate();
+			ServerLostDelegate.ExecuteIfBound();
+		}
 
-void FPortalRpcLocator::HandleMessage(const FPortalRpcServer& Message, const TSharedRef<IMessageContext, ESPMode::ThreadSafe>& Context)
-{
-	LastServerResponse = FDateTime::UtcNow();
+		// @todo sarge: implement actual product GUID
+		MessageEndpoint->Publish(new FPortalRpcLocateServer(FGuid(), GEngineVersion.ToString()), EMessageScope::Network);
 
-	FMessageAddress NewServerAddress;
+		return true;
+	}
 
-	if (FMessageAddress::Parse(Message.ServerAddress, NewServerAddress) && (NewServerAddress != ServerAddress))
+	FPortalRpcLocatorImpl()
+		: LastServerResponse(FDateTime::MinValue())
 	{
-		ServerAddress = NewServerAddress;
-		ServerLocatedDelegate.ExecuteIfBound();
+		MessageEndpoint = FMessageEndpoint::Builder("FPortalRpcLocator")
+			.Handling<FPortalRpcServer>(this, &FPortalRpcLocatorImpl::HandleMessage);
+
+		TickerHandle = FTicker::GetCoreTicker().AddTicker(FTickerDelegate::CreateRaw(this, &FPortalRpcLocatorImpl::HandleTicker), PORTAL_RPC_LOCATE_INTERVAL);
 	}
-}
 
+private:
 
-bool FPortalRpcLocator::HandleTicker(float DeltaTime)
-{
-	if (ServerAddress.IsValid() && (FDateTime::UtcNow() - LastServerResponse > PORTAL_RPC_LOCATE_TIMEOUT))
-	{
-		ServerAddress.Invalidate();
-		ServerLostDelegate.ExecuteIfBound();
-	}
+	/** Time at which the RPC server last responded. */
+	FDateTime LastServerResponse;
+
+	/** Message endpoint. */
+	TSharedPtr<FMessageEndpoint, ESPMode::ThreadSafe> MessageEndpoint;
+
+	/** The message address of the located RPC server, or invalid if no server available. */
+	FMessageAddress ServerAddress;
+
+	/** A delegate that is executed when an RPC server has been located. */
+	FSimpleDelegate ServerLocatedDelegate;
 
-	// @todo sarge: implement actual product GUID
-	MessageEndpoint->Publish(new FPortalRpcLocateServer(FGuid(), GEngineVersion.ToString()), EMessageScope::Network);
+	/** A delegate that is executed when the RPC server has been lost. */
+	FSimpleDelegate ServerLostDelegate;
 
-	return true;
+	/** Handle to the registered ticker. */
+	FDelegateHandle TickerHandle;
+
+	friend FPortalRpcLocatorFactory;
+};
+
+
+TSharedRef<IPortalRpcLocator> FPortalRpcLocatorFactory::Create()
+{
+	return MakeShareable(new FPortalRpcLocatorImpl());
 }
diff --git a/Engine/Source/Runtime/Portal/Rpc/Private/PortalRpcLocator.h b/Engine/Source/Runtime/Portal/Rpc/Private/PortalRpcLocator.h
--- a/Engine/Source/Runtime/Portal/Rpc/Private/PortalRpcLocator.h
+++ b/Engine/Source/Runtime/Portal/Rpc/Private/PortalRpcLocator.h
@@ -1,59 +1,10 @@
 // Copyright 1998-2015 Epic Games, Inc. All Rights Reserved.
-
 #pragma once
 
-#include "IMessageContext.h"
-#include "IPortalRpcLocator.h"
-
-
-class FMessageEndpoint;
-struct FPortalRpcServer;
+class IPortalRpcLocator;
 
-
-class FPortalRpcLocator
-	: public IPortalRpcLocator
+class FPortalRpcLocatorFactory
 {
 public:
-
-	/** Default constructor. */
-	FPortalRpcLocator();
-
-	/** Virtual destructor. */
-	virtual ~FPortalRpcLocator();
-
-public:
-
-	// IPortalRpcLocator interface
-
-	virtual const FMessageAddress& GetServerAddress() const override;
-	virtual FSimpleDelegate& OnServerLocated() override;
-	virtual FSimpleDelegate& OnServerLost() override;
-
-private:
-
-	/** Handles FPortalRpcServer messages. */
-	void HandleMessage(const FPortalRpcServer& Message, const TSharedRef<IMessageContext, ESPMode::ThreadSafe>& Context);
-
-	/** Handles the ticker. */
-	bool HandleTicker(float DeltaTime);
-
-private:
-
-	/** Time at which the RPC server last responded. */
-	FDateTime LastServerResponse;
-
-	/** Message endpoint. */
-	TSharedPtr<FMessageEndpoint, ESPMode::ThreadSafe> MessageEndpoint;
-
-	/** The message address of the located RPC server, or invalid if no server available. */
-	FMessageAddress ServerAddress;
-
-	/** A delegate that is executed when an RPC server has been located. */
-	FSimpleDelegate ServerLocatedDelegate;
-
-	/** A delegate that is executed when the RPC server has been lost. */
-	FSimpleDelegate ServerLostDelegate;
-
-	/** Handle to the registered ticker. */
-	FDelegateHandle TickerHandle;
+	static TSharedRef<IPortalRpcLocator> Create();
 };
diff --git a/Engine/Source/Runtime/Portal/Rpc/Private/PortalRpcModule.cpp b/Engine/Source/Runtime/Portal/Rpc/Private/PortalRpcModule.cpp
--- a/Engine/Source/Runtime/Portal/Rpc/Private/PortalRpcModule.cpp
+++ b/Engine/Source/Runtime/Portal/Rpc/Private/PortalRpcModule.cpp
@@ -2,18 +2,16 @@
 
 #include "PortalRpcPrivatePCH.h"
 #include "IPortalRpcModule.h"
+#include "PortalRpcLocator.h"
+#include "PortalRpcResponder.h"
+#include "PortalRpcServer.h"
 
-
-/**
- * Implements the PortalRpc module.
- */
 class FPortalRpcModule
 	: public IPortalRpcModule
 {
 public:
 
 	// IModuleInterface interface
-
 	virtual void StartupModule() override { }
 	virtual void ShutdownModule() override { }
 
@@ -26,14 +24,19 @@ public:
 
 	virtual TSharedRef<IPortalRpcLocator> CreateLocator() override
 	{
-		return MakeShareable(new FPortalRpcLocator);
+		return FPortalRpcLocatorFactory::Create();
 	}
 
 	virtual TSharedRef<IPortalRpcResponder> CreateResponder() override
 	{
-		return MakeShareable(new FPortalRpcResponder);
+		return FPortalRpcResponderFactory::Create();
 	}
-};
 
+	virtual TSharedRef<IPortalRpcServer> CreateServer() override
+	{
+		return FPortalRpcServerFactory::Create();
+	}
+
+};
 
 IMPLEMENT_MODULE(FPortalRpcModule, PortalRpc);
diff --git a/Engine/Source/Runtime/Portal/Rpc/Private/PortalRpcPrivatePCH.h b/Engine/Source/Runtime/Portal/Rpc/Private/PortalRpcPrivatePCH.h
--- a/Engine/Source/Runtime/Portal/Rpc/Private/PortalRpcPrivatePCH.h
+++ b/Engine/Source/Runtime/Portal/Rpc/Private/PortalRpcPrivatePCH.h
@@ -1,5 +1,4 @@
 // Copyright 1998-2015 Epic Games, Inc. All Rights Reserved.
-
 #pragma once
 
 
@@ -26,6 +25,3 @@
 /* Private includes
  *****************************************************************************/
 
-#include "PortalRpcLocator.h"
-#include "PortalRpcMessages.h"
-#include "PortalRpcResponder.h"
diff --git a/Engine/Source/Runtime/Portal/Rpc/Private/PortalRpcResponder.cpp b/Engine/Source/Runtime/Portal/Rpc/Private/PortalRpcResponder.cpp
--- a/Engine/Source/Runtime/Portal/Rpc/Private/PortalRpcResponder.cpp
+++ b/Engine/Source/Runtime/Portal/Rpc/Private/PortalRpcResponder.cpp
@@ -1,45 +1,80 @@
 // Copyright 1998-2015 Epic Games, Inc. All Rights Reserved.
 
 #include "PortalRpcPrivatePCH.h"
+#include "PortalRpcResponder.h"
+#include "IPortalRpcResponder.h"
 #include "IMessageRpcServer.h"
+#include "IPortalRpcServer.h"
+#include "PortalRpcMessages.h"
 #include "ModuleManager.h"
 
+class FPortalRpcResponderImpl
+	: public IPortalRpcResponder
+{
+public:
 
-/* FPortalRpcResponder structors
- *****************************************************************************/
+	virtual ~FPortalRpcResponderImpl() { }
 
-FPortalRpcResponder::FPortalRpcResponder()
-{
-	MessageEndpoint = FMessageEndpoint::Builder("FPortalRpcResponder")
-		.Handling<FPortalRpcLocateServer>(this, &FPortalRpcResponder::HandleMessage);
+public:
 
-	if (MessageEndpoint.IsValid())
+	// IPortalRpcResponder interface
+	virtual FOnPortalRpcLookup& OnLookup() override
 	{
-		MessageEndpoint->Subscribe<FPortalRpcLocateServer>();
+		return LookupDelegate;
 	}
-}
 
+private:
 
-/* FPortalRpcResponder callbacks
- *****************************************************************************/
-
-void FPortalRpcResponder::HandleMessage(const FPortalRpcLocateServer& Message, const TSharedRef<IMessageContext, ESPMode::ThreadSafe>& Context)
-{
-	if (!LookupDelegate.IsBound())
+	FPortalRpcResponderImpl()
 	{
-		return;
+		MessageEndpoint = FMessageEndpoint::Builder("FPortalRpcResponder")
+			.Handling<FPortalRpcLocateServer>(this, &FPortalRpcResponderImpl::HandleMessage);
+
+		if (MessageEndpoint.IsValid())
+		{
+			MessageEndpoint->Subscribe<FPortalRpcLocateServer>();
+		}
 	}
 
-	const FString ProductKey = Message.ProductId.ToString() + Message.ProductVersion;
-	TSharedPtr<IMessageRpcServer> Server = Servers.FindRef(ProductKey);
+private:
 
-	if (!Server.IsValid())
+	void HandleMessage(const FPortalRpcLocateServer& Message, const TSharedRef<IMessageContext, ESPMode::ThreadSafe>& Context)
 	{
-		Server = LookupDelegate.Execute(ProductKey);
-	}
+		if (!LookupDelegate.IsBound())
+		{
+			return;
+		}
 
-	if (Server.IsValid())
-	{
-		MessageEndpoint->Send(new FPortalRpcServer(Server->GetAddress().ToString()), Context->GetSender());
+		const FString ProductKey = Message.ProductId.ToString() + Message.ProductVersion;
+		TSharedPtr<IPortalRpcServer> Server = Servers.FindRef(ProductKey);
+
+		if (!Server.IsValid())
+		{
+			Server = LookupDelegate.Execute(ProductKey);
+		}
+
+		if (Server.IsValid())
+		{
+			Server->ConnectTo(Context->GetSender());
+		}
 	}
+
+
+private:
+
+	/** A delegate that is executed when a look-up for an RPC server occurs. */
+	FOnPortalRpcLookup LookupDelegate;
+
+	/** Message endpoint. */
+	TSharedPtr<FMessageEndpoint, ESPMode::ThreadSafe> MessageEndpoint;
+
+	/** Holds the existing RPC servers. */
+	TMap<FString, TSharedPtr<IPortalRpcServer>> Servers;
+
+	friend FPortalRpcResponderFactory;
+};
+
+TSharedRef<IPortalRpcResponder> FPortalRpcResponderFactory::Create()
+{
+	return MakeShareable(new FPortalRpcResponderImpl());
 }
diff --git a/Engine/Source/Runtime/Portal/Rpc/Private/PortalRpcResponder.h b/Engine/Source/Runtime/Portal/Rpc/Private/PortalRpcResponder.h
--- a/Engine/Source/Runtime/Portal/Rpc/Private/PortalRpcResponder.h
+++ b/Engine/Source/Runtime/Portal/Rpc/Private/PortalRpcResponder.h
@@ -1,50 +1,10 @@
 // Copyright 1998-2015 Epic Games, Inc. All Rights Reserved.
-
 #pragma once
 
-#include "IPortalRpcResponder.h"
-
-
-class FMessageEndpoint;
-struct FPortalRpcLocateServer;
-class IMessageContext;
+class IPortalRpcResponder;
 
-
-class FPortalRpcResponder
-	: public IPortalRpcResponder
+class FPortalRpcResponderFactory
 {
 public:
-
-	/** Default constructor. */
-	FPortalRpcResponder();
-
-	/** Virtual destructor. */
-	virtual ~FPortalRpcResponder() { }
-
-public:
-
-	// IPortalRpcResponder interface
-
-	virtual FOnPortalRpcLookup& OnLookup() override
-	{
-		return LookupDelegate;
-	}
-
-private:
-
-private:
-
-	/** Handles FPortalRpcLocateServer messages. */
-	void HandleMessage(const FPortalRpcLocateServer& Message, const TSharedRef<IMessageContext, ESPMode::ThreadSafe>& Context);
-
-private:
-
-	/** A delegate that is executed when a look-up for an RPC server occurs. */
-	FOnPortalRpcLookup LookupDelegate;
-
-	/** Message endpoint. */
-	TSharedPtr<FMessageEndpoint, ESPMode::ThreadSafe> MessageEndpoint;
-
-	/** Holds the existing RPC servers. */
-	TMap<FString, TSharedPtr<IMessageRpcServer>> Servers;
-};
+	static TSharedRef<IPortalRpcResponder> Create();
+};
\ No newline at end of file
diff --git a/Engine/Source/Runtime/Portal/Rpc/Private/PortalRpcServer.cpp b/Engine/Source/Runtime/Portal/Rpc/Private/PortalRpcServer.cpp
new file mode 100644
--- /dev/null
+++ b/Engine/Source/Runtime/Portal/Rpc/Private/PortalRpcServer.cpp
@@ -0,0 +1,44 @@
+// Copyright 1998-2015 Epic Games, Inc. All Rights Reserved.
+
+#include "PortalRpcPrivatePCH.h"
+#include "PortalRpcServer.h"
+#include "IPortalRpcServer.h"
+#include "PortalRpcMessages.h"
+#include "MessageRpcServer.h"
+#include "IMessageContext.h"
+
+class FPortalRpcServerImpl
+	: public FMessageRpcServer
+	, public IPortalRpcServer
+{
+public:
+
+	virtual void ConnectTo(const FMessageAddress& Address) const
+	{
+		MessageEndpoint->Send(new FPortalRpcServer(GetAddress().ToString()), Address);
+	}
+
+	virtual IMessageRpcServer* GetMessageServer()
+	{
+		return static_cast<IMessageRpcServer*>(this);
+	}
+
+public:
+
+	virtual ~FPortalRpcServerImpl() { }
+
+private:
+
+	FPortalRpcServerImpl()
+		: FMessageRpcServer()
+	{}
+
+private:
+
+	friend FPortalRpcServerFactory;
+};
+
+TSharedRef<IPortalRpcServer> FPortalRpcServerFactory::Create()
+{
+	return MakeShareable(new FPortalRpcServerImpl());
+}
diff --git a/Engine/Source/Runtime/Portal/Rpc/Private/PortalRpcServer.h b/Engine/Source/Runtime/Portal/Rpc/Private/PortalRpcServer.h
new file mode 100644
--- /dev/null
+++ b/Engine/Source/Runtime/Portal/Rpc/Private/PortalRpcServer.h
@@ -0,0 +1,10 @@
+// Copyright 1998-2015 Epic Games, Inc. All Rights Reserved.
+#pragma once
+
+class IPortalRpcServer;
+
+class FPortalRpcServerFactory
+{
+public:
+	static TSharedRef<IPortalRpcServer> Create();
+};
\ No newline at end of file
diff --git a/Engine/Source/Runtime/Portal/Rpc/Public/IPortalRpcModule.h b/Engine/Source/Runtime/Portal/Rpc/Public/IPortalRpcModule.h
--- a/Engine/Source/Runtime/Portal/Rpc/Public/IPortalRpcModule.h
+++ b/Engine/Source/Runtime/Portal/Rpc/Public/IPortalRpcModule.h
@@ -1,13 +1,11 @@
 // Copyright 1998-2015 Epic Games, Inc. All Rights Reserved.
-
 #pragma once
 
 #include "ModuleInterface.h"
 
-
 class IPortalRpcLocator;
 class IPortalRpcResponder;
-
+class IPortalRpcServer;
 
 /**
  * Interface for the PortalRpc module.
@@ -19,4 +17,5 @@ public:
 
 	virtual TSharedRef<IPortalRpcLocator> CreateLocator() = 0;
 	virtual TSharedRef<IPortalRpcResponder> CreateResponder() = 0;
+	virtual TSharedRef<IPortalRpcServer> CreateServer() = 0;
 };
diff --git a/Engine/Source/Runtime/Portal/Rpc/Public/IPortalRpcResponder.h b/Engine/Source/Runtime/Portal/Rpc/Public/IPortalRpcResponder.h
--- a/Engine/Source/Runtime/Portal/Rpc/Public/IPortalRpcResponder.h
+++ b/Engine/Source/Runtime/Portal/Rpc/Public/IPortalRpcResponder.h
@@ -1,12 +1,10 @@
 // Copyright 1998-2015 Epic Games, Inc. All Rights Reserved.
-
 #pragma once
 
-
 class IMessageRpcServer;
+class IPortalRpcServer;
 
-
-DECLARE_DELEGATE_RetVal_OneParam(TSharedPtr<IMessageRpcServer>, FOnPortalRpcLookup, const FString& /*ProductKey*/)
+DECLARE_DELEGATE_RetVal_OneParam(TSharedPtr<IPortalRpcServer>, FOnPortalRpcLookup, const FString& /*ProductKey*/)
 
 
 /**
diff --git a/Engine/Source/Runtime/Portal/Rpc/Public/IPortalRpcServer.h b/Engine/Source/Runtime/Portal/Rpc/Public/IPortalRpcServer.h
new file mode 100644
--- /dev/null
+++ b/Engine/Source/Runtime/Portal/Rpc/Public/IPortalRpcServer.h
@@ -0,0 +1,26 @@
+// Copyright 1998-2015 Epic Games, Inc. All Rights Reserved.
+#pragma once
+
+struct FMessageAddress;
+class IMessageRpcServer;
+
+/**
+ * Interface for Portal RPC server
+ */
+class IPortalRpcServer
+{
+public:
+
+	/**
+	 * Sends a FPortalRpcServer response to the specified address, effectively
+	 * allowing the two services to now communicate with each other directly.
+	 */
+	virtual void ConnectTo(const FMessageAddress& Address) const = 0;
+
+	virtual IMessageRpcServer* GetMessageServer() = 0;
+
+public:
+
+	/** Virtual destructor. */
+	~IPortalRpcServer() { }
+};
diff --git a/Engine/Source/Runtime/Portal/Services/Private/PortalServiceLocator.cpp b/Engine/Source/Runtime/Portal/Services/Private/PortalServiceLocator.cpp
--- a/Engine/Source/Runtime/Portal/Services/Private/PortalServiceLocator.cpp
+++ b/Engine/Source/Runtime/Portal/Services/Private/PortalServiceLocator.cpp
@@ -1,65 +1,101 @@
 // Copyright 1998-2015 Epic Games, Inc. All Rights Reserved.
 
 #include "PortalServicesPrivatePCH.h"
+#include "PortalServiceLocator.h"
+#include "IPortalServiceLocator.h"
 #include "IPortalServiceProvider.h"
 #include "ModuleManager.h"
+#include "TypeContainer.h"
 
-
-/* FPortalServiceLocator structors
- *****************************************************************************/
-
-FPortalServiceLocator::FPortalServiceLocator(const TSharedRef<FTypeContainer>& InServiceDependencies)
-	: ServiceDependencies(InServiceDependencies)
-{ }
-
-
-/* IPortalServiceLocator interface
- *****************************************************************************/
-
-void FPortalServiceLocator::Configure(const FString& ServiceName, const FWildcardString& ProductId, const FName& ServiceModule)
+class FPortalServiceLocatorImpl
+	: public IPortalServiceLocator
 {
-	TArray<FConfigEntry>& Entries = Configuration.FindOrAdd(ServiceName);
-	FConfigEntry Entry;
-	{
-		Entry.ProductId = ProductId;
-		Entry.ServiceModule = ServiceModule;
-	}
+public:
 
-	Entries.Add(Entry);
-}
+	~FPortalServiceLocatorImpl() { }
 
+public:
 
-TSharedPtr<IPortalService> FPortalServiceLocator::GetService(const FString& ServiceName, const FString& ProductId)
-{
-	TSharedPtr<void> Result;
-	TArray<FConfigEntry>& Entries = Configuration.FindOrAdd(ServiceName);
+	// IPortalServiceLocator interface
 
-	for (FConfigEntry& Entry : Entries)
+	virtual void Configure(const FString& ServiceName, const FWildcardString& ProductId, const FName& ServiceModule) override
 	{
-		if (!Entry.ProductId.IsMatch(ProductId))
+		TArray<FConfigEntry>& Entries = Configuration.FindOrAdd(ServiceName);
+		FConfigEntry Entry;
 		{
-			continue;
+			Entry.ProductId = ProductId;
+			Entry.ServiceModule = ServiceModule;
 		}
 
-		if (Entry.ServiceInstance.IsValid())
-		{
-			return Entry.ServiceInstance;
-		}
+		Entries.Add(Entry);
+	}
 
-		auto ServiceProvider = FModuleManager::LoadModulePtr<IPortalServiceProvider>(Entry.ServiceModule);
+	virtual TSharedPtr<IPortalService> GetService(const FString& ServiceName, const FString& ProductId) override
+	{
+		TSharedPtr<void> Result;
+		TArray<FConfigEntry>& Entries = Configuration.FindOrAdd(ServiceName);
 
-		if (ServiceProvider == nullptr)
+		for (FConfigEntry& Entry : Entries)
 		{
-			continue;
-		}
+			if (!Entry.ProductId.IsMatch(ProductId))
+			{
+				continue;
+			}
 
-		Entry.ServiceInstance = ServiceProvider->GetService(ServiceName, ServiceDependencies.ToSharedRef());
+			if (Entry.ServiceInstance.IsValid())
+			{
+				return Entry.ServiceInstance;
+			}
 
-		if (Entry.ServiceInstance.IsValid())
-		{
-			return Entry.ServiceInstance;
+			auto ServiceProvider = FModuleManager::LoadModulePtr<IPortalServiceProvider>(Entry.ServiceModule);
+
+			if (ServiceProvider == nullptr)
+			{
+				continue;
+			}
+
+			Entry.ServiceInstance = ServiceProvider->GetService(ServiceName, ServiceDependencies.ToSharedRef());
+
+			if (Entry.ServiceInstance.IsValid())
+			{
+				return Entry.ServiceInstance;
+			}
 		}
+
+		return nullptr;
 	}
 
-	return nullptr;
+private:
+	
+	/**
+	 * Create and initialize a new instance.
+	 *
+	 * @param InServiceDependencies A type container for optional service dependencies.
+	 */
+	FPortalServiceLocatorImpl(const TSharedRef<FTypeContainer>& InServiceDependencies)
+		: ServiceDependencies(InServiceDependencies)
+	{ }
+
+
+private:
+
+	struct FConfigEntry
+	{
+		FWildcardString ProductId;
+		TSharedPtr<IPortalService> ServiceInstance;
+		FName ServiceModule;
+	};
+
+	/** Holds the service configuration entries. */
+	TMap<FString, TArray<FConfigEntry>> Configuration;
+
+	/** Optional service dependencies. */
+	TSharedPtr<FTypeContainer> ServiceDependencies;
+
+	friend FPortalServiceLocatorFactory;
+};
+
+TSharedRef<IPortalServiceLocator> FPortalServiceLocatorFactory::Create(const TSharedRef<FTypeContainer>& ServiceDependencies)
+{
+	return MakeShareable(new FPortalServiceLocatorImpl(ServiceDependencies));
 }
diff --git a/Engine/Source/Runtime/Portal/Services/Private/PortalServiceLocator.h b/Engine/Source/Runtime/Portal/Services/Private/PortalServiceLocator.h
--- a/Engine/Source/Runtime/Portal/Services/Private/PortalServiceLocator.h
+++ b/Engine/Source/Runtime/Portal/Services/Private/PortalServiceLocator.h
@@ -1,48 +1,11 @@
 // Copyright 1998-2015 Epic Games, Inc. All Rights Reserved.
-
 #pragma once
 
-#include "IPortalServiceLocator.h"
-
-
-class IPortalService;
+class IPortalServiceLocator;
 class FTypeContainer;
 
-
-class FPortalServiceLocator
-	: public IPortalServiceLocator
+class FPortalServiceLocatorFactory
 {
 public:
-
-	/**
-	 * Create and initialize a new instance.
-	 *
-	 * @param InServiceDependencies A type container for optional service dependencies.
-	 */
-	FPortalServiceLocator(const TSharedRef<FTypeContainer>& InServiceDependencies);
-
-	/** Virtual destructor. */
-	~FPortalServiceLocator() { }
-
-public:
-
-	// IPortalServiceLocator interface
-
-	virtual void Configure(const FString& ServiceName, const FWildcardString& ProductId, const FName& ServiceModule) override;
-	virtual TSharedPtr<IPortalService> GetService(const FString& ServiceName, const FString& ProductId) override;
-
-private:
-
-	struct FConfigEntry
-	{
-		FWildcardString ProductId;
-		TSharedPtr<IPortalService> ServiceInstance;
-		FName ServiceModule;
-	};
-
-	/** Holds the service configuration entries. */
-	TMap<FString, TArray<FConfigEntry>> Configuration;
-
-	/** Optional service dependencies. */
-	TSharedPtr<FTypeContainer> ServiceDependencies;
+	static TSharedRef<IPortalServiceLocator> Create(const TSharedRef<FTypeContainer>& ServiceDependencies);
 };
diff --git a/Engine/Source/Runtime/Portal/Services/Private/PortalServicesModule.cpp b/Engine/Source/Runtime/Portal/Services/Private/PortalServicesModule.cpp
--- a/Engine/Source/Runtime/Portal/Services/Private/PortalServicesModule.cpp
+++ b/Engine/Source/Runtime/Portal/Services/Private/PortalServicesModule.cpp
@@ -2,6 +2,7 @@
 
 #include "PortalServicesPrivatePCH.h"
 #include "IPortalServicesModule.h"
+#include "PortalServiceLocator.h"
 #include "ModuleManager.h"
 
 
@@ -22,7 +23,7 @@ public:
 
 	virtual TSharedRef<IPortalServiceLocator> CreateLocator(const TSharedRef<FTypeContainer>& ServiceDependencies) override
 	{
-		return MakeShareable(new FPortalServiceLocator(ServiceDependencies));
+		return FPortalServiceLocatorFactory::Create(ServiceDependencies);
 	}
 
 public:
diff --git a/Engine/Source/Runtime/Portal/Services/Private/PortalServicesPrivatePCH.h b/Engine/Source/Runtime/Portal/Services/Private/PortalServicesPrivatePCH.h
--- a/Engine/Source/Runtime/Portal/Services/Private/PortalServicesPrivatePCH.h
+++ b/Engine/Source/Runtime/Portal/Services/Private/PortalServicesPrivatePCH.h
@@ -12,4 +12,3 @@
 /* Private includes
 *****************************************************************************/
 
-#include "PortalServiceLocator.h"
diff --git a/Engine/Source/Runtime/Portal/Services/Public/Application/IPortalAppWindow.h b/Engine/Source/Runtime/Portal/Services/Public/Application/IPortalAppWindow.h
deleted file mode 100644
diff --git a/Engine/Source/Runtime/Portal/Services/Public/Application/IPortalApplicationWindow.h b/Engine/Source/Runtime/Portal/Services/Public/Application/IPortalApplicationWindow.h
new file mode 100644
--- /dev/null
+++ b/Engine/Source/Runtime/Portal/Services/Public/Application/IPortalApplicationWindow.h
@@ -0,0 +1,27 @@
+// Copyright 1998-2015 Epic Games, Inc. All Rights Reserved.
+#pragma once
+
+#include "AsyncResult.h"
+#include "IPortalService.h"
+
+/**
+ * Interface for the Portal application's marketplace services.
+ */
+class IPortalApplicationWindow
+	: public IPortalService
+{
+public:
+
+	/**
+	 * Requests that the Portal navigate the user to the specified Url
+	 *
+	 * @return true on success, false otherwise.
+	 */
+	virtual TAsyncResult<bool> NavigateTo(const FString& Url) = 0;
+
+public:
+
+	/** Virtual destructor. */
+	virtual ~IPortalApplicationWindow() { }
+};
+Expose_TNameOf(IPortalApplicationWindow)
diff --git a/Engine/Source/Runtime/Portal/Services/Public/IPortalService.h b/Engine/Source/Runtime/Portal/Services/Public/IPortalService.h
--- a/Engine/Source/Runtime/Portal/Services/Public/IPortalService.h
+++ b/Engine/Source/Runtime/Portal/Services/Public/IPortalService.h
@@ -1,12 +1,12 @@
 // Copyright 1998-2015 Epic Games, Inc. All Rights Reserved.
-
 #pragma once
 
-
 class IPortalService
 {
 public:
 
 	/** Virtual destructor. */
 	virtual ~IPortalService() { }
+
+	virtual bool IsAvailable() const = 0;
 };
