Change GEngineVersion to FEngineVersion::Current()

diff --git a/Engine/Source/Runtime/Portal/Rpc/Private/PortalRpcLocator.cpp b/Engine/Source/Runtime/Portal/Rpc/Private/PortalRpcLocator.cpp
--- a/Engine/Source/Runtime/Portal/Rpc/Private/PortalRpcLocator.cpp
+++ b/Engine/Source/Runtime/Portal/Rpc/Private/PortalRpcLocator.cpp
@@ -62,7 +62,7 @@ private:
 		}
 
 		// @todo sarge: implement actual product GUID
-		MessageEndpoint->Publish(new FPortalRpcLocateServer(FGuid(), GEngineVersion.ToString()), EMessageScope::Network);
+		MessageEndpoint->Publish(new FPortalRpcLocateServer(FGuid(), FEngineVersion::Current().ToString()), EMessageScope::Network);
 
 		return true;
 	}
