Merging CIS fix from Dev-Console.

diff --git a/Engine/Source/Runtime/Portal/LauncherPlatform/LauncherPlatform.Build.cs b/Engine/Source/Runtime/Portal/LauncherPlatform/LauncherPlatform.Build.cs
--- a/Engine/Source/Runtime/Portal/LauncherPlatform/LauncherPlatform.Build.cs
+++ b/Engine/Source/Runtime/Portal/LauncherPlatform/LauncherPlatform.Build.cs
@@ -14,7 +14,7 @@ public class LauncherPlatform : ModuleRules
             }
         );
 
-		if (!Target.IsInPlatformGroup(UnrealPlatformGroup.Unix) && Target.Platform != UnrealTargetPlatform.Win32 && Target.Platform != UnrealTargetPlatform.Win64 && Target.Platform != UnrealTargetPlatform.Mac)
+        if (Target.Platform != UnrealTargetPlatform.Linux && Target.Platform != UnrealTargetPlatform.Win32 && Target.Platform != UnrealTargetPlatform.Win64 && Target.Platform != UnrealTargetPlatform.Mac)
         {
             PrecompileForTargets = PrecompileTargetsType.None;
         }
