Fix allocator inconsistencies in PortalRpcLocator.

- Flagged by ASan on Linux.

diff --git a/Engine/Source/Runtime/Portal/Rpc/Private/PortalRpcLocator.cpp b/Engine/Source/Runtime/Portal/Rpc/Private/PortalRpcLocator.cpp
--- a/Engine/Source/Runtime/Portal/Rpc/Private/PortalRpcLocator.cpp
+++ b/Engine/Source/Runtime/Portal/Rpc/Private/PortalRpcLocator.cpp
@@ -68,7 +68,9 @@ private:
 		}
 
 		// @todo sarge: implement actual product GUID
-		MessageEndpoint->Publish(new FPortalRpcLocateServer(FGuid(), EngineVersion, MacAddress, UserId), EMessageScope::Network);
+		// message is going to be deleted by FMemory::Free() (see FMessageContext destructor), so allocate it with Malloc
+		void* Memory = FMemory::Malloc(sizeof(FPortalRpcLocateServer));
+		MessageEndpoint->Publish(new(Memory) FPortalRpcLocateServer(FGuid(), EngineVersion, MacAddress, UserId), EMessageScope::Network);
 
 		return true;
 	}
