Added a verify command to the launcher module which allows the game to inform the launcher that the installation might be corrupt and to trigger a verify operation

diff --git a/Engine/Source/Runtime/Portal/LauncherCheck/Private/LauncherCheckModule.cpp b/Engine/Source/Runtime/Portal/LauncherCheck/Private/LauncherCheckModule.cpp
--- a/Engine/Source/Runtime/Portal/LauncherCheck/Private/LauncherCheckModule.cpp
+++ b/Engine/Source/Runtime/Portal/LauncherCheck/Private/LauncherCheckModule.cpp
@@ -72,6 +72,9 @@ public:
 			case ELauncherAction::AppInstaller:
 				LauncherOptions.LauncherRelativeUrl += TEXT("?action=installer");
 				break;
+			case ELauncherAction::AppVerifyInstallation:
+				LauncherOptions.LauncherRelativeUrl += TEXT("?action=verify");
+				break;
 			};
 
 			// If our payload starts with the correct encoded character, then append the string
diff --git a/Engine/Source/Runtime/Portal/LauncherCheck/Public/ILauncherCheckModule.h b/Engine/Source/Runtime/Portal/LauncherCheck/Public/ILauncherCheckModule.h
--- a/Engine/Source/Runtime/Portal/LauncherCheck/Public/ILauncherCheckModule.h
+++ b/Engine/Source/Runtime/Portal/LauncherCheck/Public/ILauncherCheckModule.h
@@ -17,7 +17,9 @@ enum class ELauncherAction : uint32
 	/** Check for an available update to the App after the launcher is started */
 	AppUpdateCheck,
 	/** Modify the Apps installation - app has to support SD and be installed */
-	AppInstaller
+	AppInstaller,
+	/** Notify the launcher that installation corruption was detected and that it should verify */
+	AppVerifyInstallation
 };
 
 /**
