Added Epic Launcher param for runningapptimeout=60. Code clean up.

diff --git a/Engine/Source/Runtime/Portal/LauncherPlatform/Public/ILauncherPlatform.h b/Engine/Source/Runtime/Portal/LauncherPlatform/Public/ILauncherPlatform.h
--- a/Engine/Source/Runtime/Portal/LauncherPlatform/Public/ILauncherPlatform.h
+++ b/Engine/Source/Runtime/Portal/LauncherPlatform/Public/ILauncherPlatform.h
@@ -5,17 +5,10 @@
 class FOpenLauncherOptions
 {
 public:
-	FOpenLauncherOptions()
-		: bInstall(false)
-		, bSilent(true)
-		, LauncherRelativeUrl()
-	{
-	}
+	FOpenLauncherOptions() {}
 
 	FOpenLauncherOptions(FString InLauncherRelativeUrl)
-		: bInstall(false)
-		, bSilent(true)
-		, LauncherRelativeUrl(InLauncherRelativeUrl)
+		: LauncherRelativeUrl(InLauncherRelativeUrl)
 	{
 		if ( !LauncherRelativeUrl.IsEmpty() )
 		{
@@ -25,7 +18,6 @@ public:
 
 	FOpenLauncherOptions(bool DoInstall, FString InLauncherRelativeUrl)
 		: bInstall(DoInstall)
-		, bSilent(true)
 		, LauncherRelativeUrl(InLauncherRelativeUrl)
 	{
 		if ( !LauncherRelativeUrl.IsEmpty() || bInstall )
@@ -37,6 +29,10 @@ public:
 	FString GetLauncherUriRequest() const
 	{
 		FString LauncherUriRequest;
+		int RunningAppTimeout = 60;
+
+		GConfig->GetInt(TEXT("LauncherPlatform"), TEXT("RunningAppTimeout"), RunningAppTimeout, GEngineIni);
+
 		if ( LauncherRelativeUrl.IsEmpty() )
 		{
 			LauncherUriRequest = TEXT("com.epicgames.launcher:");
@@ -46,31 +42,27 @@ public:
 			LauncherUriRequest = FString::Printf(TEXT("com.epicgames.launcher://%s"), *LauncherRelativeUrl);
 		}
 
+		// Append running app timeout arg.
+		if (LauncherUriRequest.Contains("?"))
+		{
+			LauncherUriRequest += TEXT("&");
+		}
+		else
+		{
+			LauncherUriRequest += TEXT("?");
+		}
+		LauncherUriRequest += FString::Printf(TEXT("runningapptimeout=%d"), RunningAppTimeout);
+
 		// Append silent query string arg.
 		if ( bSilent )
 		{
-			if ( LauncherUriRequest.Contains("?") )
-			{
-				LauncherUriRequest += TEXT("&silent=true");
-			}
-			else
-			{
-				LauncherUriRequest += TEXT("?silent=true");
-			}
+			LauncherUriRequest += TEXT("&silent=true");
 		}
 
+		// Append action string if present
 		if (!Action.IsEmpty())
 		{
-			if (LauncherUriRequest.Contains("?"))
-			{
-				LauncherUriRequest += TEXT("&");			
-			}
-			else
-			{
-				LauncherUriRequest += TEXT("?");
-			}
-
-			LauncherUriRequest += FString::Printf(TEXT("action=%s"), *Action);
+			LauncherUriRequest += FString::Printf(TEXT("&action=%s"), *Action);
 		}
 
 		return LauncherUriRequest;
@@ -78,8 +70,8 @@ public:
 
 public:
 
-	bool bInstall;
-	bool bSilent;
+	bool bInstall = false;
+	bool bSilent = true;
 	FString Action;
 	FString LauncherRelativeUrl;
 };
