-Fixed a compilation issue due to a missing header.

diff --git a/Engine/Source/Runtime/Portal/LauncherPlatform/Public/ILauncherPlatform.h b/Engine/Source/Runtime/Portal/LauncherPlatform/Public/ILauncherPlatform.h
--- a/Engine/Source/Runtime/Portal/LauncherPlatform/Public/ILauncherPlatform.h
+++ b/Engine/Source/Runtime/Portal/LauncherPlatform/Public/ILauncherPlatform.h
@@ -2,6 +2,8 @@
 
 #pragma once
 
+#include "CoreGlobals.h"
+
 class FOpenLauncherOptions
 {
 public:
