-Added proper header to fix compilation issue

diff --git a/Engine/Source/Runtime/Portal/LauncherPlatform/Public/ILauncherPlatform.h b/Engine/Source/Runtime/Portal/LauncherPlatform/Public/ILauncherPlatform.h
--- a/Engine/Source/Runtime/Portal/LauncherPlatform/Public/ILauncherPlatform.h
+++ b/Engine/Source/Runtime/Portal/LauncherPlatform/Public/ILauncherPlatform.h
@@ -2,7 +2,7 @@
 
 #pragma once
 
-#include "CoreGlobals.h"
+#include "Misc/ConfigCacheIni.h"
 
 class FOpenLauncherOptions
 {
