- Move GetEncodedExePath into LauncherPlatform so it's accessible outside of LauncherCheck

diff --git a/Engine/Source/Runtime/Portal/LauncherCheck/Private/LauncherCheckModule.cpp b/Engine/Source/Runtime/Portal/LauncherCheck/Private/LauncherCheckModule.cpp
--- a/Engine/Source/Runtime/Portal/LauncherCheck/Private/LauncherCheckModule.cpp
+++ b/Engine/Source/Runtime/Portal/LauncherCheck/Private/LauncherCheckModule.cpp
@@ -60,7 +60,7 @@ public:
 			// Construct a url to tell the launcher of this app and what we want to do with it
 			FOpenLauncherOptions LauncherOptions;
 			LauncherOptions.LauncherRelativeUrl = TEXT("apps");
-			LauncherOptions.LauncherRelativeUrl /= GetEncodedExePath();
+			LauncherOptions.LauncherRelativeUrl /= FLauncherMisc::GetEncodedExePath();
 			switch (Action)
 			{
 			case ELauncherAction::AppLaunch:
@@ -97,38 +97,6 @@ public:
 	}
 
 	virtual void ShutdownModule() override { }
-
-private:
-
-	/**
-	 * Return url encoded full path of currently running executable.
-	 */
-	FString GetEncodedExePath() const
-	{
-		FString EncodedExePath;
-
-		static const TCHAR* Delims[] = { TEXT("/") };
-		static const int32 NumDelims = UE_ARRAY_COUNT(Delims);
-
-		// Get the path to the executable, and encode it ('cos spaces, colons, etc break things)
-		FString ExePath = FPlatformProcess::BaseDir();
-
-		// Make sure it's not relative and the slashes are the right way
-		ExePath = FPaths::ConvertRelativePathToFull(ExePath);
-		ExePath.ReplaceInline(TEXT("\\"), TEXT("/"));
-
-		// Encode the path 'cos symbols like ':',' ' etc are bad
-		TArray<FString> ExeFolders;
-		ExePath.ParseIntoArray(ExeFolders, Delims, NumDelims);
-		for (const FString& ExeFolder : ExeFolders)
-		{
-			EncodedExePath /= FGenericPlatformHttp::UrlEncode(ExeFolder);
-		}
-		// Make sure it ends in a slash
-		EncodedExePath /= TEXT("");
-
-		return EncodedExePath;
-	}
 };
 
 
diff --git a/Engine/Source/Runtime/Portal/LauncherPlatform/LauncherPlatform.Build.cs b/Engine/Source/Runtime/Portal/LauncherPlatform/LauncherPlatform.Build.cs
--- a/Engine/Source/Runtime/Portal/LauncherPlatform/LauncherPlatform.Build.cs
+++ b/Engine/Source/Runtime/Portal/LauncherPlatform/LauncherPlatform.Build.cs
@@ -9,6 +9,7 @@ public class LauncherPlatform : ModuleRules
         PrivateDependencyModuleNames.AddRange(
             new string[] {
                 "Core",
+				"HTTP"
             }
         );
 
diff --git a/Engine/Source/Runtime/Portal/LauncherPlatform/Private/ILauncherPlatform.cpp b/Engine/Source/Runtime/Portal/LauncherPlatform/Private/ILauncherPlatform.cpp
new file mode 100644
--- /dev/null
+++ b/Engine/Source/Runtime/Portal/LauncherPlatform/Private/ILauncherPlatform.cpp
@@ -0,0 +1,34 @@
+// Copyright Epic Games, Inc. All Rights Reserved.
+
+#include "ILauncherPlatform.h"
+
+#include "HAL/PlatformProcess.h"
+#include "Misc/Paths.h"
+#include "PlatformHttp.h"
+
+FString FLauncherMisc::GetEncodedExePath()
+{
+	FString EncodedExePath;
+
+	static const TCHAR* Delims[] = { TEXT("/") };
+	static const int32 NumDelims = UE_ARRAY_COUNT(Delims);
+
+	// Get the path to the executable, and encode it ('cos spaces, colons, etc break things)
+	FString ExePath = FPlatformProcess::BaseDir();
+
+	// Make sure it's not relative and the slashes are the right way
+	ExePath = FPaths::ConvertRelativePathToFull(ExePath);
+	ExePath.ReplaceInline(TEXT("\\"), TEXT("/"));
+
+	// Encode the path 'cos symbols like ':',' ' etc are bad
+	TArray<FString> ExeFolders;
+	ExePath.ParseIntoArray(ExeFolders, Delims, NumDelims);
+	for (const FString& ExeFolder : ExeFolders)
+	{
+		EncodedExePath /= FPlatformHttp::UrlEncode(ExeFolder);
+	}
+	// Make sure it ends in a slash
+	EncodedExePath /= TEXT("");
+
+	return EncodedExePath;
+}
diff --git a/Engine/Source/Runtime/Portal/LauncherPlatform/Public/ILauncherPlatform.h b/Engine/Source/Runtime/Portal/LauncherPlatform/Public/ILauncherPlatform.h
--- a/Engine/Source/Runtime/Portal/LauncherPlatform/Public/ILauncherPlatform.h
+++ b/Engine/Source/Runtime/Portal/LauncherPlatform/Public/ILauncherPlatform.h
@@ -102,4 +102,11 @@ public:
 	 * @return true if the marketplace was opened, false if it is not installed or could not be installed/opened.
 	 */
 	virtual bool OpenLauncher(const FOpenLauncherOptions& Options) = 0;
+};
+
+class FLauncherMisc
+{
+public:
+	/** Return url encoded full path of currently running executable. */
+	static FString GetEncodedExePath();
 };
\ No newline at end of file
