- Add external linkage to FLauncherMisc

diff --git a/Engine/Source/Runtime/Portal/LauncherPlatform/Public/ILauncherPlatform.h b/Engine/Source/Runtime/Portal/LauncherPlatform/Public/ILauncherPlatform.h
--- a/Engine/Source/Runtime/Portal/LauncherPlatform/Public/ILauncherPlatform.h
+++ b/Engine/Source/Runtime/Portal/LauncherPlatform/Public/ILauncherPlatform.h
@@ -104,7 +104,7 @@ public:
 	virtual bool OpenLauncher(const FOpenLauncherOptions& Options) = 0;
 };
 
-class FLauncherMisc
+class LAUNCHERPLATFORM_API FLauncherMisc
 {
 public:
 	/** Return url encoded full path of currently running executable. */
