LauncherPlat (Mac): Remove deprecated flag.

diff --git a/Engine/Source/Runtime/Portal/LauncherPlatform/Private/Mac/LauncherPlatformMac.cpp b/Engine/Source/Runtime/Portal/LauncherPlatform/Private/Mac/LauncherPlatformMac.cpp
--- a/Engine/Source/Runtime/Portal/LauncherPlatform/Private/Mac/LauncherPlatformMac.cpp
+++ b/Engine/Source/Runtime/Portal/LauncherPlatform/Private/Mac/LauncherPlatformMac.cpp
@@ -26,7 +26,7 @@ bool FLauncherPlatformMac::OpenLauncher(const FOpenLauncherOptions& Options)
 		NSRunningApplication* Launcher = [RunningLaunchers objectAtIndex: 0];
 		if ( !Launcher.hidden || Options.bInstall ) // If the launcher is running, but hidden, don't activate on editor startup
 		{
-			[Launcher activateWithOptions : NSApplicationActivateAllWindows | NSApplicationActivateIgnoringOtherApps];
+			[Launcher activateWithOptions : NSApplicationActivateAllWindows];
 			FString Error;
 			FPlatformProcess::LaunchURL(*LauncherUriRequest, nullptr, &Error);
 		}
